Merging using Dev-VR_to_Release-4.21

diff --git a/Engine/Plugins/Developer/RenderDocPlugin/Source/RenderDocPlugin/Private/RenderDocPluginModule.cpp b/Engine/Plugins/Developer/RenderDocPlugin/Source/RenderDocPlugin/Private/RenderDocPluginModule.cpp
--- a/Engine/Plugins/Developer/RenderDocPlugin/Source/RenderDocPlugin/Private/RenderDocPluginModule.cpp
+++ b/Engine/Plugins/Developer/RenderDocPlugin/Source/RenderDocPlugin/Private/RenderDocPluginModule.cpp
@@ -129,7 +129,7 @@ TSharedPtr<class IInputDevice> FRenderDocPluginModule::CreateInputDevice(const T
 
 void FRenderDocPluginModule::StartupModule()
 {
-#if !UE_BUILD_SHIPPING // Disable in shipping builds
+#if WITH_EDITOR && !UE_BUILD_SHIPPING // Disable in shipping builds
 	Loader.Initialize();
 	RenderDocAPI = nullptr;
 	TickNumber = 0;
