Changed a RenderDoc error message to a warning for people that enable the plugin but don't have RenderDoc installed

diff --git a/Engine/Plugins/Developer/RenderDocPlugin/Source/RenderDocPlugin/Private/RenderDocPluginLoader.cpp b/Engine/Plugins/Developer/RenderDocPlugin/Source/RenderDocPlugin/Private/RenderDocPluginLoader.cpp
--- a/Engine/Plugins/Developer/RenderDocPlugin/Source/RenderDocPlugin/Private/RenderDocPluginLoader.cpp
+++ b/Engine/Plugins/Developer/RenderDocPlugin/Source/RenderDocPlugin/Private/RenderDocPluginLoader.cpp
@@ -140,7 +140,7 @@ void FRenderDocPluginLoader::Initialize()
 	// 4) All bets are off; aborting...
 	if (RenderDocDLL == nullptr)
 	{
-		UE_LOG(RenderDocPlugin, Error, TEXT("unable to initialize the plugin because no RenderDoc libray has been located."));
+		UE_LOG(RenderDocPlugin, Warning, TEXT("unable to initialize the plugin because no RenderDoc libray has been located."));
 		return;
 	}
 
