UE4.22 - Lambda fix
#rnx

diff --git a/Engine/Plugins/Developer/RenderDocPlugin/Source/RenderDocPlugin/Private/RenderDocPluginModule.cpp b/Engine/Plugins/Developer/RenderDocPlugin/Source/RenderDocPlugin/Private/RenderDocPluginModule.cpp
--- a/Engine/Plugins/Developer/RenderDocPlugin/Source/RenderDocPlugin/Private/RenderDocPluginModule.cpp
+++ b/Engine/Plugins/Developer/RenderDocPlugin/Source/RenderDocPlugin/Private/RenderDocPluginModule.cpp
@@ -217,10 +217,11 @@ void FRenderDocPluginModule::BeginCapture()
 
 	typedef FRenderDocPluginLoader::RENDERDOC_API_CONTEXT RENDERDOC_API_CONTEXT;
 	FRenderDocPluginModule* Plugin = this;
+	FRenderDocPluginLoader::RENDERDOC_API_CONTEXT* RenderDocAPILocal = RenderDocAPI;
 	ENQUEUE_RENDER_COMMAND(StartRenderDocCapture)(
-		[Plugin, WindowHandle, RenderDocAPI](FRHICommandListImmediate& RHICmdList)
+		[Plugin, WindowHandle, RenderDocAPILocal](FRHICommandListImmediate& RHICmdList)
 		{
-			FRenderDocFrameCapturer::BeginCapture(WindowHandle, RenderDocAPI, Plugin);
+			FRenderDocFrameCapturer::BeginCapture(WindowHandle, RenderDocAPILocal, Plugin);
 		});
 }
 
@@ -251,10 +252,11 @@ void FRenderDocPluginModule::EndCapture()
 
 	typedef FRenderDocPluginLoader::RENDERDOC_API_CONTEXT RENDERDOC_API_CONTEXT;
 	FRenderDocPluginModule* Plugin = this;
+	FRenderDocPluginLoader::RENDERDOC_API_CONTEXT* RenderDocAPILocal = RenderDocAPI;
 	ENQUEUE_RENDER_COMMAND(EndRenderDocCapture)(
-		[WindowHandle, RenderDocAPI, Plugin](FRHICommandListImmediate& RHICmdList)
+		[WindowHandle, RenderDocAPILocal, Plugin](FRHICommandListImmediate& RHICmdList)
 		{
-			FRenderDocFrameCapturer::EndCapture(WindowHandle, RenderDocAPI, Plugin);
+			FRenderDocFrameCapturer::EndCapture(WindowHandle, RenderDocAPILocal, Plugin);
 		});
 }
 
