Make RenderDocPlugin editor-only, since it depends on MainFrame.

#rnx

diff --git a/Engine/Plugins/Developer/RenderDocPlugin/RenderDocPlugin.uplugin b/Engine/Plugins/Developer/RenderDocPlugin/RenderDocPlugin.uplugin
--- a/Engine/Plugins/Developer/RenderDocPlugin/RenderDocPlugin.uplugin
+++ b/Engine/Plugins/Developer/RenderDocPlugin/RenderDocPlugin.uplugin
@@ -13,7 +13,7 @@
 	[
 		{
 			"Name" : "RenderDocPlugin",
-			"Type" : "Developer",
+			"Type" : "Editor",
 			"LoadingPhase" : "PostConfigInit",
 			"WhitelistPlatforms": [
 				"Win32",
