Change RenderDoc plugin to uncooked only, mimicing previous behavior of being initialized regardless of whether the editor is starting or not.

diff --git a/Engine/Plugins/Developer/RenderDocPlugin/RenderDocPlugin.uplugin b/Engine/Plugins/Developer/RenderDocPlugin/RenderDocPlugin.uplugin
--- a/Engine/Plugins/Developer/RenderDocPlugin/RenderDocPlugin.uplugin
+++ b/Engine/Plugins/Developer/RenderDocPlugin/RenderDocPlugin.uplugin
@@ -13,7 +13,7 @@
 	[
 		{
 			"Name" : "RenderDocPlugin",
-			"Type" : "Editor",
+			"Type" : "UncookedOnly",
 			"LoadingPhase" : "PostConfigInit",
 			"WhitelistPlatforms": [
 				"Win32",
