Changed the renderdoc settings object so that the settings are saved properly

diff --git a/Engine/Plugins/Developer/RenderDocPlugin/Source/RenderDocPlugin/Private/RenderDocPluginSettings.cpp b/Engine/Plugins/Developer/RenderDocPlugin/Source/RenderDocPlugin/Private/RenderDocPluginSettings.cpp
--- a/Engine/Plugins/Developer/RenderDocPlugin/Source/RenderDocPlugin/Private/RenderDocPluginSettings.cpp
+++ b/Engine/Plugins/Developer/RenderDocPlugin/Source/RenderDocPlugin/Private/RenderDocPluginSettings.cpp
@@ -5,11 +5,6 @@
 
 static FName DeveloperSettingsConsoleVariableMetaFName(TEXT("ConsoleVariable"));
 
-void URenderDocPluginSettings::SaveSettings()
-{
-	SaveConfig(CPF_Config, *GetClass()->GetDefaultConfigFilename(), GConfig);
-}
-
 void URenderDocPluginSettings::PostInitProperties()
 {
 	Super::PostInitProperties();
@@ -18,7 +13,6 @@ void URenderDocPluginSettings::PostInitProperties()
 	if (IsTemplate())
 	{
 		ImportConsoleVariableValues();
-		SaveSettings();
 	}
 #endif
 }
diff --git a/Engine/Plugins/Developer/RenderDocPlugin/Source/RenderDocPlugin/Private/SRenderDocPluginHelpWindow.cpp b/Engine/Plugins/Developer/RenderDocPlugin/Source/RenderDocPlugin/Private/SRenderDocPluginHelpWindow.cpp
--- a/Engine/Plugins/Developer/RenderDocPlugin/Source/RenderDocPlugin/Private/SRenderDocPluginHelpWindow.cpp
+++ b/Engine/Plugins/Developer/RenderDocPlugin/Source/RenderDocPlugin/Private/SRenderDocPluginHelpWindow.cpp
@@ -112,7 +112,6 @@ void SRenderDocPluginHelpWindow::Construct(const FArguments& InArgs)
 					{
 						URenderDocPluginSettings* Settings = GetMutableDefault<URenderDocPluginSettings>();
 						Settings->bShowHelpOnStartup = (NewState == ECheckBoxState::Checked);
-						Settings->SaveSettings();
 					})
 					[
 						SNew(STextBlock)
diff --git a/Engine/Plugins/Developer/RenderDocPlugin/Source/RenderDocPlugin/Public/RenderDocPluginSettings.h b/Engine/Plugins/Developer/RenderDocPlugin/Source/RenderDocPlugin/Public/RenderDocPluginSettings.h
--- a/Engine/Plugins/Developer/RenderDocPlugin/Source/RenderDocPlugin/Public/RenderDocPluginSettings.h
+++ b/Engine/Plugins/Developer/RenderDocPlugin/Source/RenderDocPlugin/Public/RenderDocPluginSettings.h
@@ -5,7 +5,7 @@
 #include "Engine/DeveloperSettings.h"
 #include "RenderDocPluginSettings.generated.h"
 
-UCLASS(config = RenderDocPlugin, meta = (DisplayName = "RenderDoc"))
+UCLASS(config = Engine, defaultconfig, meta = (DisplayName = "RenderDoc"))
 class RENDERDOCPLUGIN_API URenderDocPluginSettings : public UDeveloperSettings
 {
 	GENERATED_BODY()
@@ -53,9 +53,6 @@ public:
 		ConfigRestartRequired = true))
 	FString RenderDocBinaryPath;
 
-public:
-	void SaveSettings();
-
 public:
 	virtual void PostInitProperties() override;
 	virtual FName GetCategoryName() const;
