Renderdoc now works with vulkan
Added RHIGetNativeInstance which is the device instance needed by vulkan.

diff --git a/Engine/Plugins/Developer/RenderDocPlugin/Source/RenderDocPlugin/Private/RenderDocPluginModule.cpp b/Engine/Plugins/Developer/RenderDocPlugin/Source/RenderDocPlugin/Private/RenderDocPluginModule.cpp
--- a/Engine/Plugins/Developer/RenderDocPlugin/Source/RenderDocPlugin/Private/RenderDocPluginModule.cpp
+++ b/Engine/Plugins/Developer/RenderDocPlugin/Source/RenderDocPlugin/Private/RenderDocPluginModule.cpp
@@ -66,18 +66,34 @@ struct FRenderDocAsyncGraphTask : public FAsyncGraphTaskBase
 class FRenderDocFrameCapturer
 {
 public:
+	static RENDERDOC_DevicePointer GetRenderdocDevicePointer()
+	{
+		RENDERDOC_DevicePointer Device;
+		if(0 == FCString::Strcmp(GDynamicRHI->GetName(),TEXT("Vulkan")))
+		{
+			Device = GDynamicRHI->RHIGetNativeInstance();
+#ifndef RENDERDOC_DEVICEPOINTER_FROM_VKINSTANCE
+#define RENDERDOC_DEVICEPOINTER_FROM_VKINSTANCE(inst) (*((void **)(inst)))
+#endif
+			Device = RENDERDOC_DEVICEPOINTER_FROM_VKINSTANCE(Device);
+		}
+		else
+		{
+			Device = GDynamicRHI->RHIGetNativeDevice();
+		}
+		return Device;
+
+	}
 	static void BeginCapture(HWND WindowHandle, FRenderDocPluginLoader::RENDERDOC_API_CONTEXT* RenderDocAPI, FRenderDocPluginModule* Plugin)
 	{
 		UE4_GEmitDrawEvents_BeforeCapture = GetEmitDrawEvents();
 		SetEmitDrawEvents(true);
-
-		RENDERDOC_DevicePointer Device = GDynamicRHI->RHIGetNativeDevice();
-		RenderDocAPI->StartFrameCapture(Device, WindowHandle);
+		RenderDocAPI->StartFrameCapture(GetRenderdocDevicePointer(), WindowHandle);
 	}
 	static void EndCapture(HWND WindowHandle, FRenderDocPluginLoader::RENDERDOC_API_CONTEXT* RenderDocAPI, FRenderDocPluginModule* Plugin)
 	{
-		RENDERDOC_DevicePointer Device = GDynamicRHI->RHIGetNativeDevice();
-		RenderDocAPI->EndFrameCapture(Device, WindowHandle);
+		FRHICommandListExecutor::GetImmediateCommandList().SubmitCommandsAndFlushGPU();
+		RenderDocAPI->EndFrameCapture(GetRenderdocDevicePointer(), WindowHandle);
 
 		SetEmitDrawEvents(UE4_GEmitDrawEvents_BeforeCapture);
 
diff --git a/Engine/Source/Runtime/Apple/MetalRHI/Private/MetalRHI.cpp b/Engine/Source/Runtime/Apple/MetalRHI/Private/MetalRHI.cpp
--- a/Engine/Source/Runtime/Apple/MetalRHI/Private/MetalRHI.cpp
+++ b/Engine/Source/Runtime/Apple/MetalRHI/Private/MetalRHI.cpp
@@ -1088,6 +1088,12 @@ void* FMetalDynamicRHI::RHIGetNativeDevice()
 	return (void*)ImmediateContext.Context->GetDevice().GetPtr();
 }
 
+void* FMetalDynamicRHI::RHIGetNativeInstance()
+{
+	return nullptr;
+}
+
+
 uint16 FMetalDynamicRHI::RHIGetPlatformTextureMaxSampleCount()
 {
 	TArray<EMobileMSAASampleCount::Type> SamplesArray{ EMobileMSAASampleCount::Type::One, EMobileMSAASampleCount::Type::Two, EMobileMSAASampleCount::Type::Four, EMobileMSAASampleCount::Type::Eight };
diff --git a/Engine/Source/Runtime/Apple/MetalRHI/Public/MetalRHI.h b/Engine/Source/Runtime/Apple/MetalRHI/Public/MetalRHI.h
--- a/Engine/Source/Runtime/Apple/MetalRHI/Public/MetalRHI.h
+++ b/Engine/Source/Runtime/Apple/MetalRHI/Public/MetalRHI.h
@@ -137,6 +137,7 @@ public:
 	virtual void RHIVirtualTextureSetFirstMipVisible(FRHITexture2D* Texture, uint32 FirstMip) final override;
 	virtual void RHIExecuteCommandList(FRHICommandList* CmdList) final override;
 	virtual void* METALRHI_API RHIGetNativeDevice() final override;
+	virtual void* METALRHI_API RHIGetNativeInstance() final override;
 	virtual class IRHICommandContext* RHIGetDefaultContext() final override;
 	virtual IRHIComputeContext* RHIGetDefaultAsyncComputeContext() final override;
 	virtual class IRHICommandContextContainer* RHIGetCommandContextContainer(int32 Index, int32 Num) final override;
diff --git a/Engine/Source/Runtime/D3D12RHI/Private/D3D12RHI.cpp b/Engine/Source/Runtime/D3D12RHI/Private/D3D12RHI.cpp
--- a/Engine/Source/Runtime/D3D12RHI/Private/D3D12RHI.cpp
+++ b/Engine/Source/Runtime/D3D12RHI/Private/D3D12RHI.cpp
@@ -342,6 +342,12 @@ void* FD3D12DynamicRHI::RHIGetNativeDevice()
 	return (void*)GetAdapter().GetD3DDevice();
 }
 
+void* FD3D12DynamicRHI::RHIGetNativeInstance()
+{
+	return nullptr;
+}
+
+
 /**
 * Returns a supported screen resolution that most closely matches the input.
 * @param Width - Input: Desired resolution width in pixels. Output: A width that the platform supports.
diff --git a/Engine/Source/Runtime/D3D12RHI/Private/D3D12RHIPrivate.h b/Engine/Source/Runtime/D3D12RHI/Private/D3D12RHIPrivate.h
--- a/Engine/Source/Runtime/D3D12RHI/Private/D3D12RHIPrivate.h
+++ b/Engine/Source/Runtime/D3D12RHI/Private/D3D12RHIPrivate.h
@@ -361,6 +361,7 @@ public:
 	virtual void RHIVirtualTextureSetFirstMipVisible(FRHITexture2D* Texture, uint32 FirstMip) override;
 	virtual void RHIExecuteCommandList(FRHICommandList* CmdList) final override;
 	virtual void* RHIGetNativeDevice() final override;
+	virtual void* RHIGetNativeInstance() final override;
 	virtual class IRHICommandContext* RHIGetDefaultContext() final override;
 	virtual class IRHIComputeContext* RHIGetDefaultAsyncComputeContext() final override;
 	virtual class IRHICommandContextContainer* RHIGetCommandContextContainer(int32 Index, int32 Num) final override;
diff --git a/Engine/Source/Runtime/EmptyRHI/Private/EmptyCommands.cpp b/Engine/Source/Runtime/EmptyRHI/Private/EmptyCommands.cpp
--- a/Engine/Source/Runtime/EmptyRHI/Private/EmptyCommands.cpp
+++ b/Engine/Source/Runtime/EmptyRHI/Private/EmptyCommands.cpp
@@ -398,6 +398,12 @@ void* FEmptyDynamicRHI::RHIGetNativeDevice()
 	return nullptr;
 }
 
+void* FEmptyDynamicRHI::RHIGetNativeInstance()
+{
+	return nullptr;
+}
+
+
 void FEmptyDynamicRHI::RHIExecuteCommandList(FRHICommandList* CmdList)
 {
 }
diff --git a/Engine/Source/Runtime/EmptyRHI/Private/EmptyRHI.cpp b/Engine/Source/Runtime/EmptyRHI/Private/EmptyRHI.cpp
--- a/Engine/Source/Runtime/EmptyRHI/Private/EmptyRHI.cpp
+++ b/Engine/Source/Runtime/EmptyRHI/Private/EmptyRHI.cpp
@@ -179,6 +179,11 @@ void* FEmptyDynamicRHI::RHIGetNativeDevice()
 	return NULL;
 }
 
+void* FEmptyDynamicRHI::RHIGetNativeInstance()
+{
+	return NULL;
+}
+
 IRHICommandContext* FEmptyDynamicRHI::RHIGetDefaultContext()
 {
 	return this;
diff --git a/Engine/Source/Runtime/EmptyRHI/Public/EmptyRHI.h b/Engine/Source/Runtime/EmptyRHI/Public/EmptyRHI.h
--- a/Engine/Source/Runtime/EmptyRHI/Public/EmptyRHI.h
+++ b/Engine/Source/Runtime/EmptyRHI/Public/EmptyRHI.h
@@ -125,6 +125,7 @@ public:
 	virtual void RHIVirtualTextureSetFirstMipVisible(FRHITexture2D* Texture, uint32 FirstMip) final override;
 	virtual void RHIExecuteCommandList(FRHICommandList* CmdList) final override;
 	virtual void* RHIGetNativeDevice() final override;
+	virtual void* RHIGetNativeInstance() final override;
 	virtual class IRHICommandContext* RHIGetDefaultContext() final override;
 	virtual class IRHICommandContextContainer* RHIGetCommandContextContainer(int32 Index, int32 Num) final override;
 
diff --git a/Engine/Source/Runtime/NullDrv/Public/NullRHI.h b/Engine/Source/Runtime/NullDrv/Public/NullRHI.h
--- a/Engine/Source/Runtime/NullDrv/Public/NullRHI.h
+++ b/Engine/Source/Runtime/NullDrv/Public/NullRHI.h
@@ -793,6 +793,10 @@ public:
 	{ 
 		return 0; 
 	}
+	virtual void* RHIGetNativeInstance() final override
+	{
+		return 0;
+	}
 	virtual void RHIPushEvent(const TCHAR* Name, FColor Color) final override
 	{
 	}
diff --git a/Engine/Source/Runtime/OpenGLDrv/Private/OpenGLDevice.cpp b/Engine/Source/Runtime/OpenGLDrv/Private/OpenGLDevice.cpp
--- a/Engine/Source/Runtime/OpenGLDrv/Private/OpenGLDevice.cpp
+++ b/Engine/Source/Runtime/OpenGLDrv/Private/OpenGLDevice.cpp
@@ -1866,6 +1866,11 @@ void* FOpenGLDynamicRHI::RHIGetNativeDevice()
 	return PlatformDevice;
 }
 
+void* FOpenGLDynamicRHI::RHIGetNativeInstance()
+{
+	return nullptr;
+}
+
 void FOpenGLDynamicRHI::InvalidateQueries( void )
 {
 	{
diff --git a/Engine/Source/Runtime/OpenGLDrv/Public/OpenGLDrv.h b/Engine/Source/Runtime/OpenGLDrv/Public/OpenGLDrv.h
--- a/Engine/Source/Runtime/OpenGLDrv/Public/OpenGLDrv.h
+++ b/Engine/Source/Runtime/OpenGLDrv/Public/OpenGLDrv.h
@@ -438,6 +438,7 @@ public:
 	virtual void RHIVirtualTextureSetFirstMipVisible(FRHITexture2D* Texture, uint32 FirstMip) final override;
 	virtual void RHIExecuteCommandList(FRHICommandList* CmdList) final override;
 	virtual void* RHIGetNativeDevice() final override;
+	virtual void* RHIGetNativeInstance() final override;
 	virtual class IRHICommandContext* RHIGetDefaultContext() final override;
 	virtual class IRHICommandContextContainer* RHIGetCommandContextContainer(int32 Index, int32 Num) final override;
 
diff --git a/Engine/Source/Runtime/RHI/Public/DynamicRHI.h b/Engine/Source/Runtime/RHI/Public/DynamicRHI.h
--- a/Engine/Source/Runtime/RHI/Public/DynamicRHI.h
+++ b/Engine/Source/Runtime/RHI/Public/DynamicRHI.h
@@ -1133,6 +1133,13 @@ public:
 	// FlushType: Flush RHI Thread
 	virtual void* RHIGetNativeDevice() = 0;
 
+	/**
+	* Provides access to the native instance. Generally this should be avoided but is useful for third party plugins.
+	*/
+	// FlushType: Flush RHI Thread
+	virtual void* RHIGetNativeInstance() = 0;
+
+
 	// FlushType: Thread safe
 	virtual IRHICommandContext* RHIGetDefaultContext() = 0;
 
diff --git a/Engine/Source/Runtime/RHI/Public/RHICommandList.h b/Engine/Source/Runtime/RHI/Public/RHICommandList.h
--- a/Engine/Source/Runtime/RHI/Public/RHICommandList.h
+++ b/Engine/Source/Runtime/RHI/Public/RHICommandList.h
@@ -4620,7 +4620,15 @@ public:
 		 
 		return GDynamicRHI->RHIGetNativeDevice();
 	}
-	
+
+	FORCEINLINE void* GetNativeInstance()
+	{
+		QUICK_SCOPE_CYCLE_COUNTER(STAT_RHIMETHOD_GetNativeInstance_Flush);
+		ImmediateFlush(EImmediateFlushType::FlushRHIThread);
+
+		return GDynamicRHI->RHIGetNativeInstance();
+	}
+
 	FORCEINLINE class IRHICommandContext* GetDefaultContext()
 	{
 		return RHIGetDefaultContext();
@@ -5290,6 +5298,12 @@ FORCEINLINE void* RHIGetNativeDevice()
 	return FRHICommandListExecutor::GetImmediateCommandList().GetNativeDevice();
 }
 
+FORCEINLINE void* RHIGetNativeInstance()
+{
+	return FRHICommandListExecutor::GetImmediateCommandList().GetNativeInstance();
+}
+
+
 FORCEINLINE FRHIShaderLibraryRef RHICreateShaderLibrary(EShaderPlatform Platform, FString const& FilePath, FString const& Name)
 {
     return GDynamicRHI->RHICreateShaderLibrary(Platform, FilePath, Name);
diff --git a/Engine/Source/Runtime/RHI/Public/RHIValidation.h b/Engine/Source/Runtime/RHI/Public/RHIValidation.h
--- a/Engine/Source/Runtime/RHI/Public/RHIValidation.h
+++ b/Engine/Source/Runtime/RHI/Public/RHIValidation.h
@@ -1184,6 +1184,15 @@ public:
 	{
 		return RHI->RHIGetNativeDevice();
 	}
+	/**
+	* Provides access to the native instance. Generally this should be avoided but is useful for third party plugins.
+	*/
+	// FlushType: Flush RHI Thread
+	virtual void* RHIGetNativeInstance() override final
+	{
+		return RHI->RHIGetNativeInstance();
+	}
+
 
 	// FlushType: Thread safe
 	virtual IRHICommandContext* RHIGetDefaultContext() override final;
diff --git a/Engine/Source/Runtime/VulkanRHI/Private/VulkanRHI.cpp b/Engine/Source/Runtime/VulkanRHI/Private/VulkanRHI.cpp
--- a/Engine/Source/Runtime/VulkanRHI/Private/VulkanRHI.cpp
+++ b/Engine/Source/Runtime/VulkanRHI/Private/VulkanRHI.cpp
@@ -988,6 +988,11 @@ void* FVulkanDynamicRHI::RHIGetNativeDevice()
 	return (void*)Device->GetInstanceHandle();
 }
 
+void* FVulkanDynamicRHI::RHIGetNativeInstance()
+{
+	return (void*)GetInstance();
+}
+
 IRHICommandContext* FVulkanDynamicRHI::RHIGetDefaultContext()
 {
 	return &Device->GetImmediateContext();
diff --git a/Engine/Source/Runtime/VulkanRHI/Public/VulkanDynamicRHI.h b/Engine/Source/Runtime/VulkanRHI/Public/VulkanDynamicRHI.h
--- a/Engine/Source/Runtime/VulkanRHI/Public/VulkanDynamicRHI.h
+++ b/Engine/Source/Runtime/VulkanRHI/Public/VulkanDynamicRHI.h
@@ -181,6 +181,7 @@ public:
 	virtual void RHIExecuteCommandList(FRHICommandList* CmdList) final override;
 
 	virtual void* RHIGetNativeDevice() final override;
+	virtual void* RHIGetNativeInstance() final override;
 	virtual class IRHICommandContext* RHIGetDefaultContext() final override;
 	virtual class IRHIComputeContext* RHIGetDefaultAsyncComputeContext() final override;
 	virtual class IRHICommandContextContainer* RHIGetCommandContextContainer(int32 Index, int32 Num) final override;
diff --git a/Engine/Source/Runtime/Windows/D3D11RHI/Private/D3D11Device.cpp b/Engine/Source/Runtime/Windows/D3D11RHI/Private/D3D11Device.cpp
--- a/Engine/Source/Runtime/Windows/D3D11RHI/Private/D3D11Device.cpp
+++ b/Engine/Source/Runtime/Windows/D3D11RHI/Private/D3D11Device.cpp
@@ -605,4 +605,9 @@ void* FD3D11DynamicRHI::RHIGetNativeDevice()
 	return (void*)Direct3DDevice.GetReference();
 }
 
+void* FD3D11DynamicRHI::RHIGetNativeInstance()
+{
+	return nullptr;
+}
+
 
diff --git a/Engine/Source/Runtime/Windows/D3D11RHI/Private/D3D11RHIPrivate.h b/Engine/Source/Runtime/Windows/D3D11RHI/Private/D3D11RHIPrivate.h
--- a/Engine/Source/Runtime/Windows/D3D11RHI/Private/D3D11RHIPrivate.h
+++ b/Engine/Source/Runtime/Windows/D3D11RHI/Private/D3D11RHIPrivate.h
@@ -496,6 +496,7 @@ public:
 	virtual void RHIVirtualTextureSetFirstMipVisible(FRHITexture2D* Texture, uint32 FirstMip) final override;
 	virtual void RHIExecuteCommandList(FRHICommandList* CmdList) final override;
 	virtual void* RHIGetNativeDevice() final override;
+	virtual void* RHIGetNativeInstance() final override;
 	virtual class IRHICommandContext* RHIGetDefaultContext() final override;
 	virtual class IRHICommandContextContainer* RHIGetCommandContextContainer(int32 Index, int32 Num) final override;
 
