RenderDoc plugin improvements :
- in the plugin options, we can now choose a delay before the capture takes place after hitting the button, either in seconds or in ticks (default delay is still 0)
- in the plugin options, we can now choose a number of frames to capture (default is still 1). This implies the capture turns into a "capture all activity" as we can't selectively capture a single viewport across frames
- as before, the options are accessible with the command line so this also works in PIE mode
- fixed the notification that stayed on screen forever in PIE mode
- the notification shows when a delayed capture will start as well as which frame is currently captured

diff --git a/Engine/Plugins/Developer/RenderDocPlugin/Source/RenderDocPlugin/Private/RenderDocPluginModule.cpp b/Engine/Plugins/Developer/RenderDocPlugin/Source/RenderDocPlugin/Private/RenderDocPluginModule.cpp
--- a/Engine/Plugins/Developer/RenderDocPlugin/Source/RenderDocPlugin/Private/RenderDocPluginModule.cpp
+++ b/Engine/Plugins/Developer/RenderDocPlugin/Source/RenderDocPlugin/Private/RenderDocPluginModule.cpp
@@ -48,6 +48,19 @@ static TAutoConsoleVariable<int32> CVarRenderDocSaveAllInitials(
 	TEXT("0 - Disregard initial states of resources. ")
 	TEXT("1 - Always capture the initial state of all rendering resources. ")
 	TEXT("Please note that doing this will significantly increase capture size."));
+static TAutoConsoleVariable<int32> CVarRenderDocCaptureDelayInSeconds(
+	TEXT("renderdoc.CaptureDelayInSeconds"),
+	1,
+	TEXT("0 - Capture delay's unit is in frames.")
+	TEXT("1 - Capture delay's unit is in seconds."));
+static TAutoConsoleVariable<int32> CVarRenderDocCaptureDelay(
+	TEXT("renderdoc.CaptureDelay"),
+	0,
+	TEXT("If > 0, RenderDoc will trigger the capture only after this amount of time (or frames, if CaptureDelayInSeconds is false) has passed."));
+static TAutoConsoleVariable<int32> CVarRenderDocCaptureFrameCount(
+	TEXT("renderdoc.CaptureFrameCount"),
+	0,
+	TEXT("If > 0, the RenderDoc capture will encompass more than a single frame. Note: this implies that all activity in all viewports and editor windows will be captured (i.e. same as CaptureAllActivity)"));
 
 ///////////////////////////////////////////////////////////////////////////////////////////////////////////////////////////////////////////////
 // Helper classes
@@ -148,7 +161,6 @@ void FRenderDocPluginModule::StartupModule()
 #if WITH_EDITOR && !UE_BUILD_SHIPPING // Disable in shipping builds
 	Loader.Initialize();
 	RenderDocAPI = nullptr;
-	TickNumber = 0;
 
 #if WITH_EDITOR
 	EditorExtensions = nullptr;
@@ -178,7 +190,14 @@ void FRenderDocPluginModule::StartupModule()
 	check(RenderDocAPI);
 
 	IModularFeatures::Get().RegisterModularFeature(GetModularFeatureName(), this);
-	TickNumber = 0;
+	DelayedCaptureTick = 0;
+	DelayedCaptureSeconds = 0.0;
+	CaptureFrameCount = 0;
+	CaptureEndTick = 0;
+	bCaptureDelayInSeconds = false;
+	bPendingCapture = false;
+	bCaptureInProgress = false;
+	bShouldCaptureAllActivity = false;
 
 	// Setup RenderDoc settings
 	FString RenderDocCapturePath = FPaths::Combine(*FPaths::ProjectSavedDir(), TEXT("RenderDocCaptures"));
@@ -218,11 +237,7 @@ void FRenderDocPluginModule::StartupModule()
 void FRenderDocPluginModule::BeginCapture()
 {
 	UE_LOG(RenderDocPlugin, Log, TEXT("Capture frame and launch renderdoc!"));
-#if WITH_EDITOR
-	FRenderDocPluginNotification::Get().ShowNotification(LOCTEXT("CaptureNotification", "Capturing frame"));
-#else
-	GEngine->AddOnScreenDebugMessage((uint64)-1, 2.0f, FColor::Emerald, FString("RenderDoc: Capturing frame"));
-#endif
+	ShowNotification(LOCTEXT("RenderDocBeginCaptureNotification", "RenderDoc capture started"), true);
 
 	pRENDERDOC_SetCaptureOptionU32 SetOptions = Loader.RenderDocAPI->SetCaptureOptionU32;
 	int ok = SetOptions(eRENDERDOC_Option_CaptureCallstacks, CVarRenderDocCaptureCallstacks.GetValueOnAnyThread() ? 1 : 0); check(ok);
@@ -241,6 +256,21 @@ void FRenderDocPluginModule::BeginCapture()
 		});
 }
 
+bool FRenderDocPluginModule::ShouldCaptureAllActivity() const
+{
+	// capturing more than 1 frame means that we can't just capture the current viewport : 
+	return CVarRenderDocCaptureAllActivity.GetValueOnAnyThread() || (CVarRenderDocCaptureFrameCount.GetValueOnAnyThread() > 1);
+}
+
+void FRenderDocPluginModule::ShowNotification(const FText& Message, bool bForceNewNotification)
+{
+#if WITH_EDITOR
+	FRenderDocPluginNotification::Get().ShowNotification(Message, bForceNewNotification);
+#else
+	GEngine->AddOnScreenDebugMessage((uint64)-1, 2.0f, FColor::Emerald, Message.ToString());
+#endif
+}
+
 void FRenderDocPluginModule::InjectDebugExecKeybind()
 {
 	// Inject our key bind into the debug execs
@@ -274,21 +304,56 @@ void FRenderDocPluginModule::EndCapture()
 		{
 			FRenderDocFrameCapturer::EndCapture(WindowHandle, RenderDocAPILocal, Plugin);
 		});
+
+	DelayedCaptureTick = 0;
+	DelayedCaptureSeconds = 0.0;
+	CaptureFrameCount = 0;
+	CaptureEndTick = 0;
+	bPendingCapture = false;
+	bCaptureInProgress = false;
 }
 
 void FRenderDocPluginModule::CaptureFrame()
 {
-	if (CVarRenderDocCaptureAllActivity.GetValueOnAnyThread())
+	int32 FrameDelay = CVarRenderDocCaptureDelay.GetValueOnAnyThread();
+
+	// Don't do anything if we're currently already waiting for a capture to end : 
+	if (bCaptureInProgress)
+		return;
+
+	// in case there's no delay and we capture the current viewport, we can trigger the capture immediately : 
+	bShouldCaptureAllActivity = ShouldCaptureAllActivity();
+	if ((FrameDelay == 0) && !bShouldCaptureAllActivity)
 	{
-		CaptureEntireFrame();
-	}	
+		DoCaptureCurrentViewport();
+	}
 	else
 	{
-		CaptureCurrentViewport();
+		// store all CVars at beginning of capture in case they change while the capture is occuring : 
+		CaptureFrameCount = CVarRenderDocCaptureFrameCount.GetValueOnAnyThread();
+		bCaptureDelayInSeconds = CVarRenderDocCaptureDelayInSeconds.GetValueOnAnyThread();
+
+		if (bCaptureDelayInSeconds)
+		{
+			DelayedCaptureSeconds = FPlatformTime::Seconds() + (double)FrameDelay;
+		}
+		else
+		{
+			// Begin tracking the global tick counter so that the Tick() method below can
+			// identify the beginning and end of a complete engine update cycle.
+			// NOTE: GFrameCounter counts engine ticks, while GFrameNumber counts render
+			// frames. Multiple frames might get rendered in a single engine update tick.
+			// All active windows are updated, in a round-robin fashion, within a single
+			// engine tick. This includes thumbnail images for material preview, material
+			// editor previews, cascade/persona previes, etc.
+			DelayedCaptureTick = GFrameCounter + FrameDelay;
+		}
+
+		bPendingCapture = true;
 	}
 }
 
-void FRenderDocPluginModule::CaptureCurrentViewport()
+void FRenderDocPluginModule::DoCaptureCurrentViewport()
 {
 	BeginCapture();
 
@@ -317,51 +382,72 @@ void FRenderDocPluginModule::CaptureCurrentViewport()
 	EndCapture();
 }
 
-void FRenderDocPluginModule::CaptureEntireFrame()
-{
-	// Are we already in the workings of capturing an entire engine frame?
-	if (TickNumber != 0)
-		return;
-
-	// Begin tracking the global tick counter so that the Tick() method below can
-	// identify the beginning and end of a complete engine update cycle.
-	// NOTE: GFrameCounter counts engine ticks, while GFrameNumber counts render
-	// frames. Multiple frames might get rendered in a single engine update tick.
-	// All active windows are updated, in a round-robin fashion, within a single
-	// engine tick. This includes thumbnail images for material preview, material
-	// editor previews, cascade/persona previes, etc.
-	TickNumber = GFrameCounter;
-}
-
 void FRenderDocPluginModule::Tick(float DeltaTime)
 {
-	if (TickNumber == 0)
+	if (!bPendingCapture && !bCaptureInProgress)
 		return;
 
-	const uint32 TickDiff = GFrameCounter - TickNumber;
-	const uint32 MaxCount = 2;
+	if (bPendingCapture)
+	{
+		check(!bCaptureInProgress); // can't be in progress and pending at the same time
 
-	check(TickDiff <= MaxCount);
+		bool bStartCapturing = false;
+		if (bCaptureDelayInSeconds)
+		{
+			bStartCapturing = FPlatformTime::Seconds() > DelayedCaptureSeconds;
+		}
+		else
+		{
+			const int64 TickDiff = GFrameCounter - DelayedCaptureTick;
+			bStartCapturing = (TickDiff == 1);
+		}
 
-	if (TickDiff == 1)
-	{
-		BeginCapture();
+		if (bStartCapturing)
+		{
+			// are we capturing only the current viewport?
+			if (!bShouldCaptureAllActivity)
+			{
+				DoCaptureCurrentViewport();
+				check(!bCaptureInProgress && !bPendingCapture); // EndCapture must have been called
+			}
+			else
+			{
+				BeginCapture();
+				// from now on, we'll detect the end of the capture by counting ticks : 
+				CaptureEndTick = GFrameCounter + CaptureFrameCount + 1;
+				bCaptureInProgress = true;
+				bPendingCapture = false;
+			}
+		}
+		else
+		{
+			float TimeLeft = bCaptureDelayInSeconds ? (float)(DelayedCaptureSeconds - FPlatformTime::Seconds()) : (float)(DelayedCaptureTick - GFrameCounter);
+			const FText& SecondsOrFrames = bCaptureDelayInSeconds ? LOCTEXT("RenderDocSeconds", "seconds") : LOCTEXT("RenderDocFrames", "frames");
+
+			ShowNotification(LOCGEN_FORMAT_ORDERED(LOCTEXT("RenderDocPendingCaptureNotification", "RenderDoc capture starting in {0} {1}"), TimeLeft, SecondsOrFrames), false);
+		}
 	}
 
-	if (TickDiff == MaxCount)
+	const int64 TickDiff = GFrameCounter - DelayedCaptureTick;
+	if (bCaptureInProgress)
 	{
-		EndCapture();
-		TickNumber = 0;
+		check(!bPendingCapture); // can't be in progress and pending at the same time
+
+		if (GFrameCounter == CaptureEndTick)
+		{
+			EndCapture();
+		}
+		else
+		{
+			ShowNotification(LOCGEN_FORMAT_ORDERED(LOCTEXT("RenderDocCaptureInProgressNotification", "RenderDoc capturing frame #{0}"),
+				CaptureFrameCount - (CaptureEndTick - 1 - GFrameCounter)), false);
+		}
 	}
 }
 
 void FRenderDocPluginModule::StartRenderDoc(FString FrameCaptureBaseDirectory)
 {
-#if WITH_EDITOR
-	FRenderDocPluginNotification::Get().ShowNotification(LOCTEXT("LaunchNotification", "Launching RenderDoc GUI") );
-#else
-	GEngine->AddOnScreenDebugMessage((uint64)-1, 2.0f, FColor::Emerald, FString("RenderDoc: Launching RenderDoc GUI"));
-#endif
+	ShowNotification(LOCTEXT("RenderDocLaunchRenderDocNotification", "Launching RenderDoc GUI"), true);
 
 	FString NewestCapture = GetNewestCapture(FrameCaptureBaseDirectory);
 	FString ArgumentString = FString::Printf(TEXT("\"%s\""), *FPaths::ConvertRelativePathToFull(NewestCapture).Append(TEXT(".log")));
@@ -375,15 +461,12 @@ void FRenderDocPluginModule::StartRenderDoc(FString FrameCaptureBaseDirectory)
 			if (0 == PID)
 			{
 				UE_LOG(LogTemp, Error, TEXT("Could not launch RenderDoc!!"));
+				ShowNotification(LOCTEXT("RenderDocLaunchRenderDocNotificationFailure", "Failed to launch RenderDoc GUI"), true);
 			}
 		}
 	}
 
-#if WITH_EDITOR
-	FRenderDocPluginNotification::Get().ShowNotification(LOCTEXT("LaunchCompletedNotification", "RenderDoc GUI Launched!") );
-#else
-	GEngine->AddOnScreenDebugMessage((uint64)-1, 2.0f, FColor::Emerald, FString("RenderDoc: GUI Launched!"));
-#endif
+	ShowNotification(LOCTEXT("RenderDocLaunchRenderDocNotificationCompleted", "RenderDoc GUI Launched!"), true);
 }
 
 FString FRenderDocPluginModule::GetNewestCapture(FString BaseDirectory)
diff --git a/Engine/Plugins/Developer/RenderDocPlugin/Source/RenderDocPlugin/Private/RenderDocPluginModule.h b/Engine/Plugins/Developer/RenderDocPlugin/Source/RenderDocPlugin/Private/RenderDocPluginModule.h
--- a/Engine/Plugins/Developer/RenderDocPlugin/Source/RenderDocPlugin/Private/RenderDocPluginModule.h
+++ b/Engine/Plugins/Developer/RenderDocPlugin/Source/RenderDocPlugin/Private/RenderDocPluginModule.h
@@ -35,16 +35,26 @@ private:
 
 	void BeginCapture();
 	void EndCapture();
-	void CaptureCurrentViewport();	
-	void CaptureEntireFrame();
+	void DoCaptureCurrentViewport();	
 
 	/** Injects a debug key bind into the local player so that the hot key works the same in game */
 	void InjectDebugExecKeybind();
 
+	bool ShouldCaptureAllActivity() const;
+
+	void ShowNotification(const FText& Message, bool bForceNewNotification);
+
 private:
 	FRenderDocPluginLoader Loader;
 	FRenderDocPluginLoader::RENDERDOC_API_CONTEXT* RenderDocAPI;
-	uint32 TickNumber; // Tracks the frame count (tick number) for a full frame capture	
+	uint64 DelayedCaptureTick; // Tracks on which frame a delayed capture should trigger, if any (when bCaptureDelayInSeconds == false)
+	double DelayedCaptureSeconds; // Tracks at which time a delayed capture should trigger, if any (when bCaptureDelayInSeconds == true)
+	uint64 CaptureFrameCount; // Tracks how many frames should be captured
+	uint64 CaptureEndTick; // Tracks the tick at which the capture currently in progress should end
+	bool bCaptureDelayInSeconds:1; // Is the capture delay in seconds or ticks?
+	bool bShouldCaptureAllActivity : 1; // true if all the whole frame should be captured, not just the active viewport
+	bool bPendingCapture : 1; // true when a delayed capture has been triggered but hasn't started yet
+	bool bCaptureInProgress:1; // true after BeginCapture() has been called and we're waiting for the end of the capture
 
 #if WITH_EDITOR
 	FRenderDocPluginEditorExtension* EditorExtensions;
diff --git a/Engine/Plugins/Developer/RenderDocPlugin/Source/RenderDocPlugin/Private/RenderDocPluginNotification.cpp b/Engine/Plugins/Developer/RenderDocPlugin/Source/RenderDocPlugin/Private/RenderDocPluginNotification.cpp
--- a/Engine/Plugins/Developer/RenderDocPlugin/Source/RenderDocPlugin/Private/RenderDocPluginNotification.cpp
+++ b/Engine/Plugins/Developer/RenderDocPlugin/Source/RenderDocPlugin/Private/RenderDocPluginNotification.cpp
@@ -12,28 +12,35 @@ FRenderDocPluginNotification::FRenderDocPluginNotification()
 {
 }
 
-void FRenderDocPluginNotification::ShowNotification(const FText& Message)
+void FRenderDocPluginNotification::ShowNotification(const FText& Message, bool bForceNewNotification)
 {
 	LastEnableTime = FPlatformTime::Seconds();
 
-	// Starting a new request! Notify the UI.
-	if (RenderDocNotificationPtr.IsValid())
+	// Just update the previous new notification if there's one : 
+	if (RenderDocNotificationPtr.IsValid() && !bForceNewNotification)
 	{
-		RenderDocNotificationPtr.Pin()->ExpireAndFadeout();
+		RenderDocNotificationPtr.Pin()->SetText(Message);
 	}
+	else
+	{
+		// cancel the previous notification, if any, before starting out a new one : 
+		if (RenderDocNotificationPtr.IsValid())
+		{
+			RenderDocNotificationPtr.Pin()->ExpireAndFadeout();
+		}
 
-	FNotificationInfo Info (Message);
-	Info.bFireAndForget = false;
-
-	// Setting fade out and expire time to 0 as the expire message is currently very obnoxious
-	Info.FadeOutDuration = 1.0f;
-	Info.ExpireDuration = 0.0f;
+		FNotificationInfo Info(Message);
+		Info.bFireAndForget = false;
 
-	RenderDocNotificationPtr = FSlateNotificationManager::Get().AddNotification(Info);
+		// Setting fade out and expire time to 0 as the expire message is currently very obnoxious
+		Info.FadeOutDuration = 1.0f;
+		Info.ExpireDuration = 0.0f;
 
-	if (RenderDocNotificationPtr.IsValid())
-	{
-		RenderDocNotificationPtr.Pin()->SetCompletionState(SNotificationItem::CS_Pending);
+		RenderDocNotificationPtr = FSlateNotificationManager::Get().AddNotification(Info);
+		if (RenderDocNotificationPtr.IsValid())
+		{
+			RenderDocNotificationPtr.Pin()->SetCompletionState(SNotificationItem::CS_Pending);
+		}
 	}
 }
 
diff --git a/Engine/Plugins/Developer/RenderDocPlugin/Source/RenderDocPlugin/Private/RenderDocPluginNotification.h b/Engine/Plugins/Developer/RenderDocPlugin/Source/RenderDocPlugin/Private/RenderDocPluginNotification.h
--- a/Engine/Plugins/Developer/RenderDocPlugin/Source/RenderDocPlugin/Private/RenderDocPluginNotification.h
+++ b/Engine/Plugins/Developer/RenderDocPlugin/Source/RenderDocPlugin/Private/RenderDocPluginNotification.h
@@ -5,9 +5,9 @@
 #if WITH_EDITOR
 
 #include "Framework/Notifications/NotificationManager.h"
-#include "TickableEditorObject.h"
+#include "Tickable.h"
 
-class FRenderDocPluginNotification : public FTickableEditorObject
+class FRenderDocPluginNotification : public FTickableGameObject
 {
 public:
 	static FRenderDocPluginNotification& Get()
@@ -16,17 +16,16 @@ public:
 		return Instance;
 	}
 
-	void ShowNotification(const FText& Message);
+	void ShowNotification(const FText& Message, bool bForceNewNotification);
 	void HideNotification();
 
 protected:
-	/** FTickableEditorObject interface */
+	/** FTickableGameObject interface */
 	virtual void Tick(float DeltaTime);
-	virtual ETickableTickType GetTickableTickType() const
-	{
-		return ETickableTickType::Always;
-	}
+	virtual ETickableTickType GetTickableTickType() const { return ETickableTickType::Always; }
 	virtual TStatId GetStatId() const override;
+	virtual bool IsTickableWhenPaused() const override { return true; }
+	virtual bool IsTickableInEditor() const override { return true; }
 
 private:
 	FRenderDocPluginNotification();
diff --git a/Engine/Plugins/Developer/RenderDocPlugin/Source/RenderDocPlugin/Public/RenderDocPluginSettings.h b/Engine/Plugins/Developer/RenderDocPlugin/Source/RenderDocPlugin/Public/RenderDocPluginSettings.h
--- a/Engine/Plugins/Developer/RenderDocPlugin/Source/RenderDocPlugin/Public/RenderDocPluginSettings.h
+++ b/Engine/Plugins/Developer/RenderDocPlugin/Source/RenderDocPlugin/Public/RenderDocPluginSettings.h
@@ -35,6 +35,26 @@ public:
 		ConfigRestartRequired = false))
 	uint32 bSaveAllInitials : 1;
 
+	UPROPERTY(config, EditAnywhere, Category = "Frame Capture Settings", meta = (
+		ConsoleVariable = "renderdoc.CaptureDelayInSeconds", DisplayName = "Capture delay in seconds",
+		ToolTip = "If checked, the capture delay's unit will be in seconds instead of frames.",
+		ConfigRestartRequired = false))
+		uint32 bCaptureDelayInSeconds : 1;
+
+	UPROPERTY(config, EditAnywhere, Category = "Frame Capture Settings", meta = (
+		ConsoleVariable = "renderdoc.CaptureDelay", DisplayName = "Capture delay",
+		ToolTip = "If > 0, RenderDoc will trigger the capture only after this amount of frames/seconds has passed.",
+		ClampMin = 0,
+		ConfigRestartRequired = false))
+	int32 CaptureDelay;
+
+	UPROPERTY(config, EditAnywhere, Category = "Frame Capture Settings", meta = (
+		ConsoleVariable = "renderdoc.CaptureFrameCount", DisplayName = "Capture frame count",
+		ToolTip = "If > 1, the RenderDoc capture will encompass more than a single frame. Note: this implies that all activity in all viewports and editor windows will be captured (i.e. same as CaptureAllActivity)",
+		ClampMin = 1,
+		ConfigRestartRequired = false))
+		int32 CaptureFrameCount;
+
 	UPROPERTY(config, EditAnywhere, Category = "Advanced Settings", meta = (
 		ConsoleVariable = "renderdoc.ShowHelpOnStartup", DisplayName = "Show help on startup",
 		ToolTip = "If checked, a help window will be shown on startup.",
