Non-unity build fix

#rnx

diff --git a/Engine/Plugins/Developer/RenderDocPlugin/Source/RenderDocPlugin/Private/RenderDocPluginModule.cpp b/Engine/Plugins/Developer/RenderDocPlugin/Source/RenderDocPlugin/Private/RenderDocPluginModule.cpp
--- a/Engine/Plugins/Developer/RenderDocPlugin/Source/RenderDocPlugin/Private/RenderDocPluginModule.cpp
+++ b/Engine/Plugins/Developer/RenderDocPlugin/Source/RenderDocPlugin/Private/RenderDocPluginModule.cpp
@@ -331,7 +331,7 @@ void FRenderDocPluginModule::CaptureFrame()
 	{
 		// store all CVars at beginning of capture in case they change while the capture is occuring : 
 		CaptureFrameCount = CVarRenderDocCaptureFrameCount.GetValueOnAnyThread();
-		bCaptureDelayInSeconds = CVarRenderDocCaptureDelayInSeconds.GetValueOnAnyThread();
+		bCaptureDelayInSeconds = CVarRenderDocCaptureDelayInSeconds.GetValueOnAnyThread() > 0;
 
 		if (bCaptureDelayInSeconds)
 		{
