Only add key binding for RenderDoc plugin uniquely into Input.ini file.
See UDN: https://udn.unrealengine.com/questions/549197/enabling-renderdoc-plugin-invalidate-the-cook-all.html

#rnx

diff --git a/Engine/Plugins/Developer/RenderDocPlugin/Source/RenderDocPlugin/Private/RenderDocPluginModule.cpp b/Engine/Plugins/Developer/RenderDocPlugin/Source/RenderDocPlugin/Private/RenderDocPluginModule.cpp
--- a/Engine/Plugins/Developer/RenderDocPlugin/Source/RenderDocPlugin/Private/RenderDocPluginModule.cpp
+++ b/Engine/Plugins/Developer/RenderDocPlugin/Source/RenderDocPlugin/Private/RenderDocPluginModule.cpp
@@ -291,7 +291,7 @@ void FRenderDocPluginModule::InjectDebugExecKeybind()
 	FConfigSection* Section = ConfigFile->Find(TEXT("/Script/Engine.PlayerInput"));
 	if (Section != nullptr)
 	{
-		Section->HandleAddCommand(TEXT("DebugExecBindings"), TEXT("(Key=F12,Command=\"RenderDoc.CaptureFrame\", Alt=true)"), true);
+		Section->HandleAddCommand(TEXT("DebugExecBindings"), TEXT("(Key=F12,Command=\"RenderDoc.CaptureFrame\", Alt=true)"), false);
 	}
 }
 
