Added a renderdoc.CapturePIE command which starts a PIE session and captures the specified number of frames. Useful when debugging initialization problems, as it allows the first frame of a PIE session to be captured.

diff --git a/Engine/Plugins/Developer/RenderDocPlugin/Source/RenderDocPlugin/Private/RenderDocPluginModule.cpp b/Engine/Plugins/Developer/RenderDocPlugin/Source/RenderDocPlugin/Private/RenderDocPluginModule.cpp
--- a/Engine/Plugins/Developer/RenderDocPlugin/Source/RenderDocPlugin/Private/RenderDocPluginModule.cpp
+++ b/Engine/Plugins/Developer/RenderDocPlugin/Source/RenderDocPlugin/Private/RenderDocPluginModule.cpp
@@ -21,6 +21,7 @@
 extern UNREALED_API UEditorEngine* GEditor;
 #include "ISettingsModule.h"
 #include "ISettingsSection.h"
+#include "Editor.h"
 #endif
 
 DEFINE_LOG_CATEGORY(RenderDocPlugin);
@@ -229,7 +230,14 @@ void FRenderDocPluginModule::StartupModule()
 	static FAutoConsoleCommand CCmdRenderDocCaptureFrame = FAutoConsoleCommand(
 		TEXT("renderdoc.CaptureFrame"),
 		TEXT("Captures the rendering commands of the next frame and launches RenderDoc"),
-		FConsoleCommandDelegate::CreateRaw(this, &FRenderDocPluginModule::CaptureFrame));
+		FConsoleCommandDelegate::CreateRaw(this, &FRenderDocPluginModule::CaptureFrame)
+	);
+
+	static FAutoConsoleCommand CCmdRenderDocCapturePIE = FAutoConsoleCommand(
+		TEXT("renderdoc.CapturePIE"),
+		TEXT("Starts a PIE session and captures the specified number of frames from the start."),
+		FConsoleCommandWithArgsDelegate::CreateRaw(this, &FRenderDocPluginModule::CapturePIE)
+	);
 
 	BindCaptureCallbacks();
 
@@ -356,6 +364,30 @@ void FRenderDocPluginModule::CaptureFrame()
 	}
 }
 
+void FRenderDocPluginModule::CapturePIE(const TArray<FString>& Args)
+{
+	if (Args.Num() < 1)
+	{
+		UE_LOG(LogTemp, Error, TEXT("Usage: renderdoc.CapturePIE NumFrames"));
+		return;
+	}
+
+	int32 NumFrames = FCString::Atoi(*Args[0]);
+	if (NumFrames < 1)
+	{
+		UE_LOG(LogTemp, Error, TEXT("NumFrames must be greater than 0."));
+		return;
+	}
+
+	void* PIEViewportHandle = GetActiveWindow();
+	RenderDocAPI->SetActiveWindow(FRenderDocFrameCapturer::GetRenderdocDevicePointer(), PIEViewportHandle);
+	RenderDocAPI->TriggerMultiFrameCapture(NumFrames);
+
+	// Wait one frame before starting the PIE session, because RenderDoc will start capturing after the next
+	// present, and we want to catch the first PIE frame.
+	StartPIEDelayFrames = 1;
+}
+
 void FRenderDocPluginModule::DoCaptureCurrentViewport()
 {
 	BeginCapture();
@@ -387,6 +419,22 @@ void FRenderDocPluginModule::DoCaptureCurrentViewport()
 
 void FRenderDocPluginModule::Tick(float DeltaTime)
 {
+#if WITH_EDITOR
+	if (StartPIEDelayFrames > 0)
+	{
+		--StartPIEDelayFrames;
+	}
+	else if(StartPIEDelayFrames == 0)
+	{
+		UEditorEngine* EditorEngine = CastChecked<UEditorEngine>(GEngine);
+		FRequestPlaySessionParams SessionParams;
+		FLevelEditorModule& LevelEditorModule = FModuleManager::Get().GetModuleChecked<FLevelEditorModule>(TEXT("LevelEditor"));
+		SessionParams.DestinationSlateViewport = LevelEditorModule.GetFirstActiveViewport();
+		EditorEngine->RequestPlaySession(SessionParams);
+		StartPIEDelayFrames = -1;
+	}
+#endif
+
 	if (!bPendingCapture && !bCaptureInProgress)
 		return;
 
diff --git a/Engine/Plugins/Developer/RenderDocPlugin/Source/RenderDocPlugin/Private/RenderDocPluginModule.h b/Engine/Plugins/Developer/RenderDocPlugin/Source/RenderDocPlugin/Private/RenderDocPluginModule.h
--- a/Engine/Plugins/Developer/RenderDocPlugin/Source/RenderDocPlugin/Private/RenderDocPluginModule.h
+++ b/Engine/Plugins/Developer/RenderDocPlugin/Source/RenderDocPlugin/Private/RenderDocPluginModule.h
@@ -27,6 +27,7 @@ public:
 
 	void Tick(float DeltaTime);
 	void CaptureFrame();
+	void CapturePIE(const TArray<FString>& Args);
 	void StartRenderDoc(FString FrameCaptureBaseDirectory);
 	FString GetNewestCapture(FString BaseDirectory);
 
@@ -64,5 +65,6 @@ private:
 
 #if WITH_EDITOR
 	FRenderDocPluginEditorExtension* EditorExtensions;
+	int StartPIEDelayFrames = -1;
 #endif
 };
