Enable RenderDoc Captures for failing automation tests.

diff --git a/Engine/Plugins/Developer/RenderDocPlugin/Source/RenderDocPlugin/Private/RenderDocPluginModule.cpp b/Engine/Plugins/Developer/RenderDocPlugin/Source/RenderDocPlugin/Private/RenderDocPluginModule.cpp
--- a/Engine/Plugins/Developer/RenderDocPlugin/Source/RenderDocPlugin/Private/RenderDocPluginModule.cpp
+++ b/Engine/Plugins/Developer/RenderDocPlugin/Source/RenderDocPlugin/Private/RenderDocPluginModule.cpp
@@ -14,6 +14,7 @@
 #include "Misc/ConfigCacheIni.h"
 #include "Engine/GameViewportClient.h"
 #include "RenderCaptureInterface.h"
+#include "Misc/AutomationTest.h"
 
 #if WITH_EDITOR
 #include "UnrealClient.h"
@@ -99,23 +100,71 @@ public:
 		return Device;
 
 	}
+
 	static void BeginCapture(HWND WindowHandle, FRenderDocPluginLoader::RENDERDOC_API_CONTEXT* RenderDocAPI, FRenderDocPluginModule* Plugin)
 	{
 		UE4_GEmitDrawEvents_BeforeCapture = GetEmitDrawEvents();
 		SetEmitDrawEvents(true);
 		RenderDocAPI->StartFrameCapture(GetRenderdocDevicePointer(), WindowHandle);
 	}
-	static void EndCapture(HWND WindowHandle, FRenderDocPluginLoader::RENDERDOC_API_CONTEXT* RenderDocAPI, FRenderDocPluginModule* Plugin)
+
+	static void EndCapture(HWND WindowHandle, FRenderDocPluginLoader::RENDERDOC_API_CONTEXT* RenderDocAPI, FRenderDocPluginModule* Plugin, const FString& DestPath, FRenderDocPluginModule::ELaunchAfterCapture LaunchOption)
 	{
 		FRHICommandListExecutor::GetImmediateCommandList().SubmitCommandsAndFlushGPU();
 		RenderDocAPI->EndFrameCapture(GetRenderdocDevicePointer(), WindowHandle);
 
 		SetEmitDrawEvents(UE4_GEmitDrawEvents_BeforeCapture);
 
-		TGraphTask<FRenderDocAsyncGraphTask>::CreateTask().ConstructAndDispatchWhenReady(ENamedThreads::GameThread, [Plugin]()
+		bool bLaunchDestPath = false;
+		if (!DestPath.IsEmpty())
 		{
-			Plugin->StartRenderDoc(FPaths::Combine(*FPaths::ProjectSavedDir(), *FString("RenderDocCaptures")));
-		});
+			IPlatformFile& PlatformFileSystem = IPlatformFile::GetPlatformPhysical();
+
+			FString NewestCapturePath = Plugin->GetNewestCapture();
+			UE_LOG(RenderDocPlugin, Log, TEXT("Copying: %s to %s"), *NewestCapturePath, *DestPath);
+
+			if (PlatformFileSystem.FileExists(*DestPath))
+			{
+				PlatformFileSystem.DeleteFile(*DestPath);
+			}
+
+			if (PlatformFileSystem.CreateDirectoryTree(*FPaths::GetPath(DestPath)))
+			{
+				if (PlatformFileSystem.MoveFile(*DestPath, *NewestCapturePath))
+				{
+					bLaunchDestPath = true;
+				}
+				else
+				{
+					uint32 WriteErrorCode = FPlatformMisc::GetLastError();
+					TCHAR WriteErrorBuffer[2048];
+					FPlatformMisc::GetSystemErrorMessage(WriteErrorBuffer, 2048, WriteErrorCode);
+					UE_LOG(RenderDocPlugin, Warning, TEXT("Failed to move: %s to %s. WriteError: %u (%s)"), *NewestCapturePath, *DestPath, WriteErrorCode, WriteErrorBuffer);
+				}
+			}
+			else
+			{
+				uint32 WriteErrorCode = FPlatformMisc::GetLastError();
+				TCHAR WriteErrorBuffer[2048];
+				FPlatformMisc::GetSystemErrorMessage(WriteErrorBuffer, 2048, WriteErrorCode);
+				UE_LOG(RenderDocPlugin, Warning, TEXT("Failed to create directory tree for: %s. WriteError: %u (%s)"), *DestPath, WriteErrorCode, WriteErrorBuffer);
+			}
+		}
+
+		if (LaunchOption == FRenderDocPluginModule::ELaunchAfterCapture::Yes)
+		{	
+			TGraphTask<FRenderDocAsyncGraphTask>::CreateTask().ConstructAndDispatchWhenReady(ENamedThreads::GameThread, [Plugin, DestPath, bLaunchDestPath]()
+			{
+				if (bLaunchDestPath)
+				{
+					Plugin->StartRenderDoc(DestPath);
+				}
+				else
+				{
+					Plugin->StartRenderDoc(FString());
+				}
+			});
+		}
 	}
 
 private:
@@ -202,7 +251,7 @@ void FRenderDocPluginModule::StartupModule()
 	bShouldCaptureAllActivity = false;
 
 	// Setup RenderDoc settings
-	FString RenderDocCapturePath = FPaths::Combine(*FPaths::ProjectSavedDir(), TEXT("RenderDocCaptures"));
+	FString RenderDocCapturePath = FPaths::ProjectSavedDir() / TEXT("RenderDocCaptures");
 	if (!IFileManager::Get().DirectoryExists(*RenderDocCapturePath))
 	{
 		IFileManager::Get().MakeDirectory(*RenderDocCapturePath, true);
@@ -303,17 +352,17 @@ void FRenderDocPluginModule::InjectDebugExecKeybind()
 	}
 }
 
-void FRenderDocPluginModule::EndCapture()
+void FRenderDocPluginModule::EndCapture(void* HWnd, const FString& DestPath, ELaunchAfterCapture LaunchOption)
 {
-	HWND WindowHandle = GetActiveWindow();
+	HWND WindowHandle = (HWnd) ? reinterpret_cast<HWND>(HWnd) : GetActiveWindow();
 
 	typedef FRenderDocPluginLoader::RENDERDOC_API_CONTEXT RENDERDOC_API_CONTEXT;
 	FRenderDocPluginModule* Plugin = this;
 	FRenderDocPluginLoader::RENDERDOC_API_CONTEXT* RenderDocAPILocal = RenderDocAPI;
 	ENQUEUE_RENDER_COMMAND(EndRenderDocCapture)(
-		[WindowHandle, RenderDocAPILocal, Plugin](FRHICommandListImmediate& RHICmdList)
+		[WindowHandle, RenderDocAPILocal, Plugin, DestPath, LaunchOption](FRHICommandListImmediate& RHICmdList)
 		{
-			FRenderDocFrameCapturer::EndCapture(WindowHandle, RenderDocAPILocal, Plugin);
+			FRenderDocFrameCapturer::EndCapture(WindowHandle, RenderDocAPILocal, Plugin, DestPath, LaunchOption);
 		});
 
 	DelayedCaptureTick = 0;
@@ -324,22 +373,29 @@ void FRenderDocPluginModule::EndCapture()
 	bCaptureInProgress = false;
 }
 
-void FRenderDocPluginModule::CaptureFrame()
+void FRenderDocPluginModule::CaptureFrame(FViewport* InViewport, const FString& DestPath, ELaunchAfterCapture LaunchOption)
 {
 	int32 FrameDelay = CVarRenderDocCaptureDelay.GetValueOnAnyThread();
 
 	// Don't do anything if we're currently already waiting for a capture to end : 
 	if (bCaptureInProgress)
+	{
 		return;
+	}
 
 	// in case there's no delay and we capture the current viewport, we can trigger the capture immediately : 
 	bShouldCaptureAllActivity = ShouldCaptureAllActivity();
 	if ((FrameDelay == 0) && !bShouldCaptureAllActivity)
 	{
-		DoCaptureCurrentViewport();
+		DoCaptureCurrentViewport(InViewport, DestPath, LaunchOption);
 	}
 	else
 	{
+		if (InViewport || !DestPath.IsEmpty() || LaunchOption != ELaunchAfterCapture::Yes)
+		{
+			UE_LOG(RenderDocPlugin, Warning, TEXT("Deferred captures only support default params. Passed in params will be ignored."));
+		}
+
 		// store all CVars at beginning of capture in case they change while the capture is occuring : 
 		CaptureFrameCount = CVarRenderDocCaptureFrameCount.GetValueOnAnyThread();
 		bCaptureDelayInSeconds = CVarRenderDocCaptureDelayInSeconds.GetValueOnAnyThread() > 0;
@@ -388,19 +444,23 @@ void FRenderDocPluginModule::CapturePIE(const TArray<FString>& Args)
 	StartPIEDelayFrames = 1;
 }
 
-void FRenderDocPluginModule::DoCaptureCurrentViewport()
+void FRenderDocPluginModule::DoCaptureCurrentViewport(FViewport* InViewport, const FString& DestPath, ELaunchAfterCapture LaunchOption)
 {
 	BeginCapture();
 
 	// infer the intended viewport to intercept/capture:
-	FViewport* Viewport (nullptr);
+	FViewport* Viewport = InViewport;
+
 	check(GEngine);
-	if (GEngine->GameViewport)
+	if (!Viewport && GEngine->GameViewport)
 	{
 		check(GEngine->GameViewport->Viewport);
 		if (GEngine->GameViewport->Viewport->HasFocus())
+		{
 			Viewport = GEngine->GameViewport->Viewport;
+		}
 	}
+
 #if WITH_EDITOR
 	if (!Viewport && GEditor)
 	{
@@ -411,10 +471,11 @@ void FRenderDocPluginModule::DoCaptureCurrentViewport()
 		Viewport = GEditor->GetActiveViewport();
 	}
 #endif
+
 	check(Viewport);
 	Viewport->Draw(true);
 
-	EndCapture();
+	EndCapture(Viewport->GetWindow(), DestPath, LaunchOption);
 }
 
 void FRenderDocPluginModule::Tick(float DeltaTime)
@@ -458,7 +519,7 @@ void FRenderDocPluginModule::Tick(float DeltaTime)
 			// are we capturing only the current viewport?
 			if (!bShouldCaptureAllActivity)
 			{
-				DoCaptureCurrentViewport();
+				DoCaptureCurrentViewport(nullptr, FString(), ELaunchAfterCapture::Yes);
 				check(!bCaptureInProgress && !bPendingCapture); // EndCapture must have been called
 			}
 			else
@@ -486,7 +547,7 @@ void FRenderDocPluginModule::Tick(float DeltaTime)
 
 		if (GFrameCounter == CaptureEndTick)
 		{
-			EndCapture();
+			EndCapture(nullptr, FString(), ELaunchAfterCapture::Yes);
 		}
 		else
 		{
@@ -496,14 +557,18 @@ void FRenderDocPluginModule::Tick(float DeltaTime)
 	}
 }
 
-void FRenderDocPluginModule::StartRenderDoc(FString FrameCaptureBaseDirectory)
+void FRenderDocPluginModule::StartRenderDoc(FString LaunchPath)
 {
 	ShowNotification(LOCTEXT("RenderDocLaunchRenderDocNotification", "Launching RenderDoc GUI"), true);
 
-	FString NewestCapture = GetNewestCapture(FrameCaptureBaseDirectory);
-	FString ArgumentString = FString::Printf(TEXT("\"%s\""), *FPaths::ConvertRelativePathToFull(NewestCapture).Append(TEXT(".log")));
+	if (LaunchPath.IsEmpty())
+	{
+		LaunchPath = GetNewestCapture();
+	}
+
+	FString ArgumentString = FString::Printf(TEXT("\"%s\""), *LaunchPath);
 
-	if (!NewestCapture.IsEmpty())
+	if (!LaunchPath.IsEmpty())
 	{
 		if (!RenderDocAPI->IsRemoteAccessConnected())
 		{
@@ -520,7 +585,7 @@ void FRenderDocPluginModule::StartRenderDoc(FString FrameCaptureBaseDirectory)
 	ShowNotification(LOCTEXT("RenderDocLaunchRenderDocNotificationCompleted", "RenderDoc GUI Launched!"), true);
 }
 
-FString FRenderDocPluginModule::GetNewestCapture(FString BaseDirectory)
+FString FRenderDocPluginModule::GetNewestCapture()
 {
 	char LogFile[512];
 	uint64_t Timestamp;
@@ -535,7 +600,7 @@ FString FRenderDocPluginModule::GetNewestCapture(FString BaseDirectory)
 		Index++;
 	}
 	
-	return OutString;
+	return FPaths::ConvertRelativePathToFull(OutString);
 }
 
 void FRenderDocPluginModule::ShutdownModule()
@@ -575,7 +640,7 @@ void FRenderDocPluginModule::EndCaptureBracket(FRHICommandListImmediate* RHIComm
 		{
 			TGraphTask<FRenderDocAsyncGraphTask>::CreateTask().ConstructAndDispatchWhenReady(ENamedThreads::GameThread, [this]()
 			{
-				StartRenderDoc(FPaths::Combine(*FPaths::ProjectSavedDir(), *FString("RenderDocCaptures")));
+				StartRenderDoc(FString());
 			});
 		}
 	});
@@ -593,6 +658,20 @@ void FRenderDocPluginModule::BindCaptureCallbacks()
 			EndCaptureBracket(RHICommandList);
 		})
 	);
+
+	if (FAutomationTestFramework::Get().OnCaptureFrameTrace.IsBound())
+	{
+		UE_LOG(RenderDocPlugin, Warning, TEXT("Automation OnCaptureFrameTrace delegate already bound. RenderDoc capture will not be bound."));
+	}
+	else
+	{
+		FAutomationTestFramework::Get().OnCaptureFrameTrace.BindLambda(
+			[](const FString& DestPath, FViewport* InViewport)
+			{
+				FRenderDocPluginModule& PluginModule = FModuleManager::GetModuleChecked<FRenderDocPluginModule>("RenderDocPlugin");
+				PluginModule.CaptureFrame(InViewport, DestPath, ELaunchAfterCapture::No);
+			});
+	}
 }
 
 void FRenderDocPluginModule::UnBindCaptureCallbacks()
diff --git a/Engine/Plugins/Developer/RenderDocPlugin/Source/RenderDocPlugin/Private/RenderDocPluginModule.h b/Engine/Plugins/Developer/RenderDocPlugin/Source/RenderDocPlugin/Private/RenderDocPluginModule.h
--- a/Engine/Plugins/Developer/RenderDocPlugin/Source/RenderDocPlugin/Private/RenderDocPluginModule.h
+++ b/Engine/Plugins/Developer/RenderDocPlugin/Source/RenderDocPlugin/Private/RenderDocPluginModule.h
@@ -21,22 +21,29 @@ DECLARE_LOG_CATEGORY_EXTERN(RenderDocPlugin, Log, All);
 
 class FRenderDocPluginModule : public IRenderDocPlugin
 {
-public:	
+public:
+	enum class ELaunchAfterCapture
+	{
+		No,
+		Yes,
+	};
+
 	virtual void StartupModule() override;
 	virtual void ShutdownModule() override;
 
 	void Tick(float DeltaTime);
-	void CaptureFrame();
+	void CaptureFrame(FViewport* Viewport, const FString& DestPath, ELaunchAfterCapture LaunchOption);
+	void CaptureFrame() { CaptureFrame(nullptr, FString(), ELaunchAfterCapture::Yes); }
 	void CapturePIE(const TArray<FString>& Args);
-	void StartRenderDoc(FString FrameCaptureBaseDirectory);
-	FString GetNewestCapture(FString BaseDirectory);
+	void StartRenderDoc(FString CapturePath);
+	FString GetNewestCapture();
 
 private:
 	virtual TSharedPtr<class IInputDevice> CreateInputDevice(const TSharedRef<FGenericApplicationMessageHandler>& InMessageHandler) override;
 
 	void BeginCapture();
-	void EndCapture();
-	void DoCaptureCurrentViewport();	
+	void EndCapture(void* HWnd, const FString& DestPath, ELaunchAfterCapture LaunchOption);
+	void DoCaptureCurrentViewport(FViewport* Viewport, const FString& DestPath, ELaunchAfterCapture LaunchOption);
 
 	/** Injects a debug key bind into the local player so that the hot key works the same in game */
 	void InjectDebugExecKeybind();
diff --git a/Engine/Source/Developer/AutomationController/Private/AutomationControllerManger.cpp b/Engine/Source/Developer/AutomationController/Private/AutomationControllerManger.cpp
--- a/Engine/Source/Developer/AutomationController/Private/AutomationControllerManger.cpp
+++ b/Engine/Source/Developer/AutomationController/Private/AutomationControllerManger.cpp
@@ -1036,6 +1036,7 @@ void FAutomationControllerManager::HandleReceivedScreenShot(const FAutomationWor
 
 	bool bTree = true;
 	FString FileName = ScreenshotIncomingFolder / Message.ScreenShotName;
+	FString TraceFileName = FPaths::ChangeExtension(ScreenshotIncomingFolder / Message.ScreenShotName, TEXT(".rdc"));
 
 	FString DirectoryPath = FPaths::GetPath(FileName);
 
@@ -1054,6 +1055,17 @@ void FAutomationControllerManager::HandleReceivedScreenShot(const FAutomationWor
 		return;
 	}
 
+	if (Message.FrameTrace.Num() > 0)
+	{
+		if (!FFileHelper::SaveArrayToFile(Message.FrameTrace, *TraceFileName))
+		{
+			uint32 WriteErrorCode = FPlatformMisc::GetLastError();
+			TCHAR WriteErrorBuffer[2048];
+			FPlatformMisc::GetSystemErrorMessage(WriteErrorBuffer, 2048, WriteErrorCode);
+			UE_LOG(LogAutomationController, Warning, TEXT("Faild to save frame trace to %s. WriteError: %u (%s)"), *TraceFileName, WriteErrorCode, WriteErrorBuffer);
+		}
+	}
+
 	// TODO Automation There is identical code in, Engine\Source\Runtime\AutomationWorker\Private\AutomationWorkerModule.cpp,
 	// need to move this code into common area.
 
diff --git a/Engine/Source/Developer/FunctionalTesting/Private/ScreenshotFunctionalTestBase.cpp b/Engine/Source/Developer/FunctionalTesting/Private/ScreenshotFunctionalTestBase.cpp
--- a/Engine/Source/Developer/FunctionalTesting/Private/ScreenshotFunctionalTestBase.cpp
+++ b/Engine/Source/Developer/FunctionalTesting/Private/ScreenshotFunctionalTestBase.cpp
@@ -123,7 +123,10 @@ void AScreenshotFunctionalTestBase::OnScreenShotCaptured(int32 InSizeX, int32 In
 	check(GEngine->GameViewport);
 
 	GEngine->GameViewport->OnScreenshotCaptured().RemoveAll(this);
+
 #if WITH_AUTOMATION_TESTS
+	TArray<uint8> CapturedFrameTrace = AutomationCommon::CaptureFrameTrace(GetWorld()->GetName(), GetName());
+
 	FAutomationScreenshotData Data = AutomationCommon::BuildScreenshotData(GetWorld()->GetName(), GetName(), InSizeX, InSizeY);
 
 	// Copy the relevant data into the metadata for the screenshot.
@@ -147,7 +150,7 @@ void AScreenshotFunctionalTestBase::OnScreenShotCaptured(int32 InSizeX, int32 In
 		FAutomationTestFramework::Get().OnScreenshotCompared.AddUObject(this, &AScreenshotFunctionalTestBase::OnComparisonComplete);
 	}
 
-	FAutomationTestFramework::Get().OnScreenshotCaptured().ExecuteIfBound(InImageData, Data);
+	FAutomationTestFramework::Get().OnScreenshotAndTraceCaptured().ExecuteIfBound(InImageData, CapturedFrameTrace, Data);
 
 	UE_LOG(LogScreenshotFunctionalTest, Log, TEXT("Screenshot captured as %s"), *Data.ScreenshotName);
 #endif
diff --git a/Engine/Source/Developer/ScreenShotComparisonTools/Private/ScreenShotManager.cpp b/Engine/Source/Developer/ScreenShotComparisonTools/Private/ScreenShotManager.cpp
--- a/Engine/Source/Developer/ScreenShotComparisonTools/Private/ScreenShotManager.cpp
+++ b/Engine/Source/Developer/ScreenShotComparisonTools/Private/ScreenShotManager.cpp
@@ -213,6 +213,8 @@ FImageComparisonResult FScreenShotManager::CompareScreenshot(FString ExistingIma
 		FString IncomingFile = ( ScreenshotUnapprovedFolder / ComparisonResult.IncomingFile );
 		FString IncomingMetadataFile = FPaths::ChangeExtension(IncomingFile, ".json");
 
+		FString IncomingTraceFile = FPaths::ChangeExtension(IncomingFile, TEXT(".rdc"));
+
 		if ( IFileManager::Get().Copy(*( ReportFolder / TEXT("Approved.png")), *ApprovedFile, true, true) == COPY_OK )
 		{
 			IFileManager::Get().Copy(*( ReportFolder / TEXT("Approved.json")), *ApprovedMetadataFile, true, true);
@@ -225,6 +227,18 @@ FImageComparisonResult FScreenShotManager::CompareScreenshot(FString ExistingIma
 			ComparisonResult.ReportIncomingFile = TEXT("Incoming.png");
 		}
 
+		if (ComparisonResult.AreSimilar())
+		{
+			IFileManager::Get().Delete(*IncomingTraceFile, false, true);
+		}
+		else
+		{
+			if (IFileManager::Get().FileExists(*IncomingTraceFile))
+			{
+				IFileManager::Get().Copy(*(ReportFolder / TEXT("Incoming.rdc")), *IncomingTraceFile, true, true);
+			}
+		}
+
 		if ( IFileManager::Get().Copy(*( ReportFolder / TEXT("Delta.png")), *( ScreenshotDeltaFolder / ComparisonResult.ComparisonFile ), true, true) == COPY_OK )
 		{
 			ComparisonResult.ReportComparisonFile = TEXT("Delta.png");
diff --git a/Engine/Source/Runtime/AutomationMessages/Public/AutomationWorkerMessages.h b/Engine/Source/Runtime/AutomationMessages/Public/AutomationWorkerMessages.h
--- a/Engine/Source/Runtime/AutomationMessages/Public/AutomationWorkerMessages.h
+++ b/Engine/Source/Runtime/AutomationMessages/Public/AutomationWorkerMessages.h
@@ -662,6 +662,10 @@ struct FAutomationWorkerScreenImage
 	UPROPERTY(EditAnywhere, Category="Message")
 	TArray<uint8> ScreenImage;
 
+	/** The frame trace data. */
+	UPROPERTY(EditAnywhere, Category = "Message")
+	TArray<uint8> FrameTrace;
+
 	/** The screen shot name. */
 	UPROPERTY(EditAnywhere, Category="Message")
 	FString ScreenShotName;
diff --git a/Engine/Source/Runtime/AutomationWorker/Private/AutomationWorkerModule.cpp b/Engine/Source/Runtime/AutomationWorker/Private/AutomationWorkerModule.cpp
--- a/Engine/Source/Runtime/AutomationWorker/Private/AutomationWorkerModule.cpp
+++ b/Engine/Source/Runtime/AutomationWorker/Private/AutomationWorkerModule.cpp
@@ -375,6 +375,7 @@ void FAutomationWorkerModule::HandlePreTestingEvent()
 {
 #if WITH_ENGINE
 	FAutomationTestFramework::Get().OnScreenshotCaptured().BindRaw(this, &FAutomationWorkerModule::HandleScreenShotCapturedWithName);
+	FAutomationTestFramework::Get().OnScreenshotAndTraceCaptured().BindRaw(this, &FAutomationWorkerModule::HandleScreenShotAndTraceCapturedWithName);
 #endif
 }
 
@@ -382,6 +383,7 @@ void FAutomationWorkerModule::HandlePreTestingEvent()
 void FAutomationWorkerModule::HandlePostTestingEvent()
 {
 #if WITH_ENGINE
+	FAutomationTestFramework::Get().OnScreenshotAndTraceCaptured().Unbind();
 	FAutomationTestFramework::Get().OnScreenshotCaptured().Unbind();
 #endif
 }
@@ -413,6 +415,11 @@ void FAutomationWorkerModule::HandlePerformanceDataRetrieved(const FAutomationWo
 
 #if WITH_ENGINE
 void FAutomationWorkerModule::HandleScreenShotCapturedWithName(const TArray<FColor>& RawImageData, const FAutomationScreenshotData& Data)
+{
+	HandleScreenShotAndTraceCapturedWithName(RawImageData, TArray<uint8>(), Data);
+}
+
+void FAutomationWorkerModule::HandleScreenShotAndTraceCapturedWithName(const TArray<FColor>& RawImageData, const TArray<uint8>& CapturedFrameTrace, const FAutomationScreenshotData& Data)
 {
 #if WITH_AUTOMATION_TESTS
 	int32 NewHeight = Data.Height;
@@ -424,12 +431,13 @@ void FAutomationWorkerModule::HandleScreenShotCapturedWithName(const TArray<FCol
 	FAutomationScreenshotMetadata Metadata(Data);
 		
 	// Send the screen shot if we have a target
-	if( TestRequesterAddress.IsValid() )
+	if (TestRequesterAddress.IsValid())
 	{
 		FAutomationWorkerScreenImage* Message = new FAutomationWorkerScreenImage();
 
 		Message->ScreenShotName = Data.ScreenshotName;
 		Message->ScreenImage = CompressedBitmap;
+		Message->FrameTrace = CapturedFrameTrace;
 		Message->Metadata = Metadata;
 
 		UE_LOG(LogAutomationWorker, Log, TEXT("Sending screenshot %s"), *Message->ScreenShotName);
@@ -442,6 +450,7 @@ void FAutomationWorkerModule::HandleScreenShotCapturedWithName(const TArray<FCol
 		const bool bTree = true;
 
 		FString LocalFile = AutomationCommon::GetLocalPathForScreenshot(Data.ScreenshotName);
+		FString LocalTraceFile = FPaths::ChangeExtension(LocalFile, TEXT(".rdc"));
 		FString DirectoryPath = FPaths::GetPath(LocalFile);
 
 		UE_LOG(LogAutomationWorker, Log, TEXT("Saving screenshot %s as %s"),*Data.ScreenshotName, *LocalFile);
@@ -461,6 +470,17 @@ void FAutomationWorkerModule::HandleScreenShotCapturedWithName(const TArray<FCol
 			return;
 		}
 
+		if (CapturedFrameTrace.Num() > 0)
+		{
+			if (!FFileHelper::SaveArrayToFile(CapturedFrameTrace, *LocalTraceFile))
+			{
+				uint32 WriteErrorCode = FPlatformMisc::GetLastError();
+				TCHAR WriteErrorBuffer[2048];
+				FPlatformMisc::GetSystemErrorMessage(WriteErrorBuffer, 2048, WriteErrorCode);
+				UE_LOG(LogAutomationWorker, Warning, TEXT("Failed to save frame trace to %s. WriteError: %u (%s)"), *LocalTraceFile, WriteErrorCode, WriteErrorBuffer);
+			}
+		}
+
 		FString Json;
 		if ( FJsonObjectConverter::UStructToJsonObjectString(Metadata, Json) )
 		{
diff --git a/Engine/Source/Runtime/AutomationWorker/Private/AutomationWorkerModule.h b/Engine/Source/Runtime/AutomationWorker/Private/AutomationWorkerModule.h
--- a/Engine/Source/Runtime/AutomationWorker/Private/AutomationWorkerModule.h
+++ b/Engine/Source/Runtime/AutomationWorker/Private/AutomationWorkerModule.h
@@ -115,6 +115,9 @@ private:
 #if WITH_ENGINE
 	/** Invoked when we have screen shot to send. */
 	void HandleScreenShotCapturedWithName(const TArray<FColor>& RawImageData, const FAutomationScreenshotData& Data);
+
+	/** Invoked when we have screen shot and frame trace to send. */
+	void HandleScreenShotAndTraceCapturedWithName(const TArray<FColor>& RawImageData, const TArray<uint8>& CapturedFrameTrace, const FAutomationScreenshotData& Data);
 #endif
 
 	/** Dispatches analytics events to the data collector. */
diff --git a/Engine/Source/Runtime/Core/Private/Misc/AutomationTest.cpp b/Engine/Source/Runtime/Core/Private/Misc/AutomationTest.cpp
--- a/Engine/Source/Runtime/Core/Private/Misc/AutomationTest.cpp
+++ b/Engine/Source/Runtime/Core/Private/Misc/AutomationTest.cpp
@@ -656,6 +656,11 @@ FOnTestScreenshotCaptured& FAutomationTestFramework::OnScreenshotCaptured()
 	return TestScreenshotCapturedDelegate;
 }
 
+FOnTestScreenshotAndTraceCaptured& FAutomationTestFramework::OnScreenshotAndTraceCaptured()
+{
+	return TestScreenshotAndTraceCapturedDelegate;
+}
+
 void FAutomationTestFramework::PrepForAutomationTests()
 {
 	check(!GIsAutomationTesting);
diff --git a/Engine/Source/Runtime/Core/Public/Misc/AutomationTest.h b/Engine/Source/Runtime/Core/Public/Misc/AutomationTest.h
--- a/Engine/Source/Runtime/Core/Public/Misc/AutomationTest.h
+++ b/Engine/Source/Runtime/Core/Public/Misc/AutomationTest.h
@@ -663,6 +663,10 @@ struct CORE_API FAutomationScreenshotCompareResults
  */
 DECLARE_DELEGATE_TwoParams(FOnTestScreenshotCaptured, const TArray<FColor>&, const FAutomationScreenshotData&);
 
+DECLARE_DELEGATE_ThreeParams(FOnTestScreenshotAndTraceCaptured, const TArray<FColor>&, const TArray<uint8>&, const FAutomationScreenshotData&);
+
+DECLARE_DELEGATE_TwoParams(FOnCaptureFrameTrace, const FString& /*DestPath*/, class FViewport* /*Optional Capture Viewport*/);
+
 DECLARE_MULTICAST_DELEGATE_OneParam(FOnTestScreenshotComparisonComplete, const FAutomationScreenshotCompareResults& /*CompareResults*/);
 
 DECLARE_MULTICAST_DELEGATE_TwoParams(FOnTestDataRetrieved, bool /*bWasNew*/, const FString& /*JsonData*/);
@@ -691,6 +695,9 @@ public:
 	/** The final call related to screenshots, after they've been taken, and after they've been compared (or not if automation isn't running). */
 	FSimpleMulticastDelegate OnScreenshotTakenAndCompared;
 
+	/** Called when a frame trace should be captured */
+	FOnCaptureFrameTrace OnCaptureFrameTrace;
+
 	/**
 	 * Return the singleton instance of the framework.
 	 *
@@ -836,6 +843,11 @@ public:
 	 */
 	FOnTestScreenshotCaptured& OnScreenshotCaptured();
 
+	/**
+	 * Accessor for delegate called when a png screenshot is captured and a frame trace
+	 */
+	FOnTestScreenshotAndTraceCaptured& OnScreenshotAndTraceCaptured();
+
 	/**
 	 * Sets forcing smoke tests.
 	 */
@@ -1018,6 +1030,9 @@ private:
 	/** Delegate called at the end of the frame when a screenshot is captured and a .png is requested */
 	FOnTestScreenshotCaptured TestScreenshotCapturedDelegate;
 
+	/** Delegate called at the end of the frame when a screenshot and frame trace is captured and a .png is requested */
+	FOnTestScreenshotAndTraceCaptured TestScreenshotAndTraceCapturedDelegate;
+
 	/** Forces running smoke tests */
 	bool bForceSmokeTests;
 
diff --git a/Engine/Source/Runtime/Engine/Private/Tests/AutomationCommon.cpp b/Engine/Source/Runtime/Engine/Private/Tests/AutomationCommon.cpp
--- a/Engine/Source/Runtime/Engine/Private/Tests/AutomationCommon.cpp
+++ b/Engine/Source/Runtime/Engine/Private/Tests/AutomationCommon.cpp
@@ -28,6 +28,13 @@ DEFINE_LOG_CATEGORY_STATIC(LogEngineAutomationLatentCommand, Log, All);
 DEFINE_LOG_CATEGORY(LogEditorAutomationTests);
 DEFINE_LOG_CATEGORY(LogEngineAutomationTests);
 
+static TAutoConsoleVariable<int32> CVarAutomationAllowFrameTraceCapture(
+	TEXT("AutomationAllowFrameTraceCapture"),
+	1,
+	TEXT("Allow automation to capture frame traces."),
+	ECVF_Default
+	);
+
 //declare static variable
 FOnEditorAutomationMapLoad AutomationCommon::OnEditorAutomationMapLoad;
 
@@ -138,6 +145,43 @@ namespace AutomationCommon
 
 		return Data;
 	}
+
+	TArray<uint8> CaptureFrameTrace(const FString& MapOrContext, const FString& TestName)
+	{
+		TArray<uint8> FrameTrace;
+
+		if (CVarAutomationAllowFrameTraceCapture.GetValueOnGameThread() != 0 && FAutomationTestFramework::Get().OnCaptureFrameTrace.IsBound())
+		{
+			const FString MapAndTest = MapOrContext / FPaths::MakeValidFileName(TestName, TEXT('_'));
+			FString ScreenshotName = GetScreenshotName(MapAndTest);
+			FString TempCaptureFilePath = FPaths::ChangeExtension(FPaths::ConvertRelativePathToFull(FPaths::AutomationDir() / TEXT("Incoming/") / ScreenshotName), TEXT(".rdc"));
+
+			UE_LOG(LogEngineAutomationTests, Log, TEXT("Taking Frame Trace: %s"), *TempCaptureFilePath);
+
+			FAutomationTestFramework::Get().OnCaptureFrameTrace.Execute(TempCaptureFilePath, GEngine->GameViewport->Viewport);
+			FlushRenderingCommands();
+
+			IPlatformFile& PlatformFileSystem = IPlatformFile::GetPlatformPhysical();
+			if (PlatformFileSystem.FileExists(*TempCaptureFilePath))
+			{
+				{
+					TUniquePtr<IFileHandle> FileHandle(PlatformFileSystem.OpenRead(*TempCaptureFilePath));
+
+					int64 FileSize = FileHandle->Size();
+					FrameTrace.SetNumUninitialized(FileSize);
+					FileHandle->Read(FrameTrace.GetData(), FileSize);
+				}
+
+				PlatformFileSystem.DeleteFile(*TempCaptureFilePath);
+			}
+			else
+			{
+				UE_LOG(LogEngineAutomationTests, Warning, TEXT("Failed taking frame trace: %s"), *TempCaptureFilePath);
+			}
+		}
+
+		return FrameTrace;
+	}
 #endif
 
 	/** These save a PNG and get sent over the network */
diff --git a/Engine/Source/Runtime/Engine/Public/Tests/AutomationCommon.h b/Engine/Source/Runtime/Engine/Public/Tests/AutomationCommon.h
--- a/Engine/Source/Runtime/Engine/Public/Tests/AutomationCommon.h
+++ b/Engine/Source/Runtime/Engine/Public/Tests/AutomationCommon.h
@@ -49,6 +49,8 @@ namespace AutomationCommon
 		return OnEditorAutomationMapLoad;
 	}
 
+	ENGINE_API TArray<uint8> CaptureFrameTrace(const FString& MapOrContext, const FString& TestName);
+
 #endif
 }
 
