Removing another reference to UE4.

#rnx

diff --git a/Engine/Plugins/Developer/RenderDocPlugin/Source/RenderDocPlugin/Private/RenderDocPluginLoader.cpp b/Engine/Plugins/Developer/RenderDocPlugin/Source/RenderDocPlugin/Private/RenderDocPluginLoader.cpp
--- a/Engine/Plugins/Developer/RenderDocPlugin/Source/RenderDocPlugin/Private/RenderDocPluginLoader.cpp
+++ b/Engine/Plugins/Developer/RenderDocPlugin/Source/RenderDocPlugin/Private/RenderDocPluginLoader.cpp
@@ -111,7 +111,7 @@ void FRenderDocPluginLoader::Initialize()
 	{
 		// THIS WILL NEVER TRIGGER because of a sort of chicken-and-egg problem: RenderDoc Loader is a PostConfigInit
 		// plugin, and GUsingNullRHI is only initialized properly between PostConfigInit and PreLoadingScreen phases.
-		// (nevertheless, keep this comment around for future iterations of UE4)
+		// (nevertheless, keep this comment around for future iterations of UE)
 		UE_LOG(RenderDocPlugin, Display, TEXT("RenderDoc plugin will not be loaded because a null RHI (Cook Server, perhaps) is being used."));
 		return;
 	}
