RenderDocPlugin - Allow scoped captures while running commandlets
* FRenderDocPluginEditorExtension::Initialize() would fail when running from commandlets

diff --git a/Engine/Plugins/Developer/RenderDocPlugin/Source/RenderDocPlugin/Private/RenderDocPluginModule.cpp b/Engine/Plugins/Developer/RenderDocPlugin/Source/RenderDocPlugin/Private/RenderDocPluginModule.cpp
--- a/Engine/Plugins/Developer/RenderDocPlugin/Source/RenderDocPlugin/Private/RenderDocPluginModule.cpp
+++ b/Engine/Plugins/Developer/RenderDocPlugin/Source/RenderDocPlugin/Private/RenderDocPluginModule.cpp
@@ -308,7 +308,10 @@ void FRenderDocPluginModule::StartupModule()
 		FConsoleCommandWithArgsDelegate::CreateRaw(this, &FRenderDocPluginModule::CapturePIE)
 	);
 
-	EditorExtensions = new FRenderDocPluginEditorExtension(this);
+	if (!IsRunningCommandlet())
+	{
+		EditorExtensions = new FRenderDocPluginEditorExtension(this);
+	}
 #endif // WITH_EDITOR
 
 	UE_LOG(RenderDocPlugin, Log, TEXT("RenderDoc plugin is ready!"));
