RenderDoc: Don't load LevelEditor (to add toolbar extension) when not in Editor. This avoids an ensure in MainFrameModule.cpp.

#rnx

diff --git a/Engine/Plugins/Developer/RenderDocPlugin/Source/RenderDocPlugin/Private/SRenderDocPluginEditorExtension.cpp b/Engine/Plugins/Developer/RenderDocPlugin/Source/RenderDocPlugin/Private/SRenderDocPluginEditorExtension.cpp
--- a/Engine/Plugins/Developer/RenderDocPlugin/Source/RenderDocPlugin/Private/SRenderDocPluginEditorExtension.cpp
+++ b/Engine/Plugins/Developer/RenderDocPlugin/Source/RenderDocPlugin/Private/SRenderDocPluginEditorExtension.cpp
@@ -95,16 +95,22 @@ void FRenderDocPluginEditorExtension::Initialize(FRenderDocPluginModule* ThePlug
 	FRenderDocPluginStyle::Initialize();
 	FRenderDocPluginCommands::Register();
 
-	FLevelEditorModule& LevelEditorModule = FModuleManager::LoadModuleChecked<FLevelEditorModule>("LevelEditor");
-	TSharedRef<FUICommandList> CommandBindings = LevelEditorModule.GetGlobalLevelEditorActions();
-	ExtensionManager = LevelEditorModule.GetToolBarExtensibilityManager();
-	ToolbarExtender = MakeShareable(new FExtender);
-	ToolbarExtension = ToolbarExtender->AddToolBarExtension("CameraSpeed", EExtensionHook::After, CommandBindings,
-		FToolBarExtensionDelegate::CreateLambda([this, ThePlugin](FToolBarBuilder& ToolbarBuilder)
-	{ 
-		AddToolbarExtension(ToolbarBuilder, ThePlugin); })
-	);
-	ExtensionManager->AddExtender(ToolbarExtender);
+#if WITH_EDITOR
+	if (!IsRunningGame())
+	{
+		FLevelEditorModule& LevelEditorModule = FModuleManager::LoadModuleChecked<FLevelEditorModule>("LevelEditor");
+		TSharedRef<FUICommandList> CommandBindings = LevelEditorModule.GetGlobalLevelEditorActions();
+		ExtensionManager = LevelEditorModule.GetToolBarExtensibilityManager();
+		ToolbarExtender = MakeShareable(new FExtender);
+		ToolbarExtension = ToolbarExtender->AddToolBarExtension("CameraSpeed", EExtensionHook::After, CommandBindings,
+			FToolBarExtensionDelegate::CreateLambda([this, ThePlugin](FToolBarBuilder& ToolbarBuilder)
+				{
+					AddToolbarExtension(ToolbarBuilder, ThePlugin); 
+				})
+		);
+		ExtensionManager->AddExtender(ToolbarExtender);
+	}
+#endif // WITH_EDITOR
 
 	// Would be nice to use the preprocessor definition WITH_EDITOR instead, but the user may launch a standalone the game through the editor...
 	if (GEditor != nullptr)
