Make RenderDocPlugin always load.
Now we need to use -AttachRenderDoc on command line to use the plugin.
This unifies behaviour with the PIX plugin.
Also added a plugin setting to auto attach without the command line for people who want to have RenderDoc always on.

diff --git a/Engine/Plugins/Developer/RenderDocPlugin/RenderDocPlugin.uplugin b/Engine/Plugins/Developer/RenderDocPlugin/RenderDocPlugin.uplugin
--- a/Engine/Plugins/Developer/RenderDocPlugin/RenderDocPlugin.uplugin
+++ b/Engine/Plugins/Developer/RenderDocPlugin/RenderDocPlugin.uplugin
@@ -8,7 +8,7 @@
 	"DocsURL" : "https://renderdoc.org/docs/index.html",
 	"Description" : "RenderDoc graphics debugger/profiler integration.",
 	"Category" : "Rendering",
-	"EnabledByDefault" : false,
+	"EnabledByDefault" : true,
 	"Modules" :
 	[
 		{
diff --git a/Engine/Plugins/Developer/RenderDocPlugin/Source/RenderDocPlugin/Private/RenderDocPluginLoader.cpp b/Engine/Plugins/Developer/RenderDocPlugin/Source/RenderDocPlugin/Private/RenderDocPluginLoader.cpp
--- a/Engine/Plugins/Developer/RenderDocPlugin/Source/RenderDocPlugin/Private/RenderDocPluginLoader.cpp
+++ b/Engine/Plugins/Developer/RenderDocPlugin/Source/RenderDocPlugin/Private/RenderDocPluginLoader.cpp
@@ -20,6 +20,11 @@
 
 #define LOCTEXT_NAMESPACE "RenderDocPlugin"
 
+static TAutoConsoleVariable<int32> CVarRenderDocAutoAttach(
+	TEXT("renderdoc.AutoAttach"),
+	0,
+	TEXT("RenderDoc will attach on startup."));
+
 static TAutoConsoleVariable<int32> CVarRenderDocEnableCrashHandler(
 	TEXT("renderdoc.EnableCrashHandler"),
 	0,
@@ -186,21 +191,33 @@ void FRenderDocPluginLoader::Initialize()
 	RenderDocDLL = nullptr;
 	RenderDocAPI = nullptr;
 
-	bool bDisableFrameTraceCapture = FParse::Param(FCommandLine::Get(), TEXT("DisableFrameTraceCapture"));
-	if (bDisableFrameTraceCapture)
-	{
-		UE_LOG(RenderDocPlugin, Display, TEXT("RenderDoc plugin will not be loaded because -DisableFrameTraceCapture cmd line flag."));
-		return;
-	}
+	// Since we are so early in the loading phase, we first need to load the cvars since they're not loaded at this point.
+	UE::ConfigUtilities::ApplyCVarSettingsFromIni(TEXT("/Script/RenderDocPlugin.RenderDocPluginSettings"), *GEngineIni, ECVF_SetByProjectSetting);
+
+	const bool bCapture = FParse::Param(FCommandLine::Get(), TEXT("AttachRenderDoc")) || (CVarRenderDocAutoAttach.GetValueOnAnyThread() != 0);
 
-	if (FApp::IsUnattended())
+	const bool bUnattended = FApp::IsUnattended();
+	IConsoleVariable* CVarAutomationAllowFrameTraceCapture = IConsoleManager::Get().FindConsoleVariable(TEXT("AutomationAllowFrameTraceCapture"), false);
+	const bool bUnattendedCapture = bUnattended && CVarAutomationAllowFrameTraceCapture && CVarAutomationAllowFrameTraceCapture->GetInt() != 0;
+
+	if (!bCapture && !bUnattendedCapture)
 	{
-		IConsoleVariable* CVarAutomationAllowFrameTraceCapture = IConsoleManager::Get().FindConsoleVariable(TEXT("AutomationAllowFrameTraceCapture"), false);
-		if (!CVarAutomationAllowFrameTraceCapture || CVarAutomationAllowFrameTraceCapture->GetInt() == 0)
+		if (bUnattended)
 		{
 			UE_LOG(RenderDocPlugin, Display, TEXT("RenderDoc plugin will not be loaded because AutomationAllowFrameTraceCapture cvar is set to 0."));
-			return;
 		}
+		else
+		{
+			UE_LOG(RenderDocPlugin, Display, TEXT("RenderDoc plugin will not be loaded. Use '-AttachRenderDoc' on the cmd line or enable 'renderdoc.AutoAttach' in the plugin settings."));
+		}
+		return;
+	}
+
+	const bool bDisableFrameTraceCapture = FParse::Param(FCommandLine::Get(), TEXT("DisableFrameTraceCapture"));
+	if (bDisableFrameTraceCapture)
+	{
+		UE_LOG(RenderDocPlugin, Display, TEXT("RenderDoc plugin will not be loaded because -DisableFrameTraceCapture cmd line flag."));
+		return;
 	}
 
 	if (GUsingNullRHI)
@@ -215,8 +232,7 @@ void FRenderDocPluginLoader::Initialize()
 	// Look for a renderdoc.dll somewhere in the system:
 	UE_LOG(RenderDocPlugin, Log, TEXT("locating RenderDoc library (%s)..."), *RenderDocDllName);
 
-	// 1) Check the Game configuration files. Since we are so early in the loading phase, we first need to load the cvars since they're not loaded at this point:
-	UE::ConfigUtilities::ApplyCVarSettingsFromIni(TEXT("/Script/RenderDocPlugin.RenderDocPluginSettings"), *GEngineIni, ECVF_SetByProjectSetting);
+	// 1) Get binary path  from CVar:
 	RenderDocDLL = LoadAndCheckRenderDocLibrary(RenderDocAPI, CVarRenderDocBinaryPath.GetValueOnAnyThread());
 
 #if PLATFORM_WINDOWS
diff --git a/Engine/Plugins/Developer/RenderDocPlugin/Source/RenderDocPlugin/Public/RenderDocPluginSettings.h b/Engine/Plugins/Developer/RenderDocPlugin/Source/RenderDocPlugin/Public/RenderDocPluginSettings.h
--- a/Engine/Plugins/Developer/RenderDocPlugin/Source/RenderDocPlugin/Public/RenderDocPluginSettings.h
+++ b/Engine/Plugins/Developer/RenderDocPlugin/Source/RenderDocPlugin/Public/RenderDocPluginSettings.h
@@ -53,7 +53,13 @@ public:
 		ToolTip = "If > 1, the RenderDoc capture will encompass more than a single frame. Note: this implies that all activity in all viewports and editor windows will be captured (i.e. same as CaptureAllActivity)",
 		ClampMin = 1,
 		ConfigRestartRequired = false))
-		int32 CaptureFrameCount;
+	int32 CaptureFrameCount;
+
+	UPROPERTY(config, EditAnywhere, Category = "Advanced Settings", meta = (
+		ConsoleVariable = "renderdoc.AutoAttach", DisplayName = "Auto attach on startup",
+		ToolTip = "If checked, then RenderDoc will attach on startup. Otherwise the command line argument '-AttachRenderDoc' is required."
+		))
+	uint32 bAutoAttach : 1;
 
 	UPROPERTY(config, EditAnywhere, Category = "Advanced Settings", meta = (
 		ConsoleVariable = "renderdoc.ShowHelpOnStartup", DisplayName = "Show help on startup",
