Update RenderDoc to better interact with the new viewport toolbar.
- Invert icon color to white so that it is consistent with other iconography.

diff --git a/Engine/Plugins/Developer/RenderDocPlugin/Source/RenderDocPlugin/Private/SRenderDocPluginEditorExtension.cpp b/Engine/Plugins/Developer/RenderDocPlugin/Source/RenderDocPlugin/Private/SRenderDocPluginEditorExtension.cpp
--- a/Engine/Plugins/Developer/RenderDocPlugin/Source/RenderDocPlugin/Private/SRenderDocPluginEditorExtension.cpp
+++ b/Engine/Plugins/Developer/RenderDocPlugin/Source/RenderDocPlugin/Private/SRenderDocPluginEditorExtension.cpp
@@ -23,6 +23,8 @@
 #include "RHI.h"
 #include "Widgets/Input/SButton.h"
 #include "Widgets/Images/SImage.h"
+#include "ViewportToolbar/UnrealEdViewportToolbar.h"
+#include "ToolMenus.h"
 
 extern UNREALED_API UEditorEngine* GEditor;
 
@@ -36,7 +38,7 @@ public:
 	void Construct(const FArguments& Args)
 	{
 		FSlateIcon IconBrush = FSlateIcon(FRenderDocPluginStyle::Get()->GetStyleSetName(), "RenderDocPlugin.CaptureFrameIcon");
-	
+		
 		ChildSlot
 		[
 			SNew(SButton)
@@ -80,6 +82,8 @@ FRenderDocPluginEditorExtension::~FRenderDocPluginEditorExtension()
 	{
 		ExtensionManager.Reset();
 	}
+	
+	UToolMenus::UnregisterOwner(this);
 }
 
 void FRenderDocPluginEditorExtension::Initialize(FRenderDocPluginModule* ThePlugin)
@@ -110,6 +114,8 @@ void FRenderDocPluginEditorExtension::Initialize(FRenderDocPluginModule* ThePlug
 				})
 		);
 		ExtensionManager->AddExtender(ToolbarExtender);
+		
+		ExtendToolbar();
 	}
 #endif // WITH_EDITOR
 
@@ -135,15 +141,50 @@ void FRenderDocPluginEditorExtension::Initialize(FRenderDocPluginModule* ThePlug
 	}
 }
 
+void FRenderDocPluginEditorExtension::ExtendToolbar()
+{
+	FToolMenuOwnerScoped ScopedOwner(this);
+
+	UToolMenu* Menu = UToolMenus::Get()->ExtendMenu("LevelEditor.ViewportToolbar");
+	
+	FToolMenuSection& RightSection = Menu->FindOrAddSection("Right");
+	FToolMenuEntry& Entry = RightSection.AddMenuEntry(
+		"RenderDocButton",
+		FText::GetEmpty(),
+		FRenderDocPluginCommands::Get().CaptureFrameCommand->GetDescription(),
+		FSlateIcon(FRenderDocPluginStyle::Get()->GetStyleSetName(), "RenderDocPlugin.CaptureFrameIcon"),
+		FUIAction(FExecuteAction::CreateLambda([]
+		{
+			FPlayWorldCommands::GlobalPlayWorldActions->GetActionForCommand(FRenderDocPluginCommands::Get().CaptureFrameCommand)->Execute();
+		}))
+	);
+	
+	Entry.InsertPosition.Position = EToolMenuInsertType::First;
+}
+
 void FRenderDocPluginEditorExtension::AddToolbarExtension(FToolBarBuilder& ToolbarBuilder, FRenderDocPluginModule* ThePlugin)
 {
 #define LOCTEXT_NAMESPACE "LevelEditorToolBar"
 
 	UE_LOG(RenderDocPlugin, Verbose, TEXT("Attaching toolbar extension..."));
-	ToolbarBuilder.AddSeparator();
-
-	ToolbarBuilder.BeginSection("RenderdocPlugin");
-	ToolbarBuilder.AddWidget(SNew(SRenderDocCaptureButton));
+	ToolbarBuilder.BeginSection("RenderdocPlugin", false);
+	
+	TAttribute<EVisibility> Visibility;
+	Visibility.BindLambda([]
+	{
+		return UE::UnrealEd::ShowOldViewportToolbars() ? EVisibility::Visible : EVisibility::Collapsed;
+	});
+	
+	ToolbarBuilder.AddSeparator(NAME_None, Visibility);
+	
+	ToolbarBuilder.AddWidget(
+		SNew(SRenderDocCaptureButton),
+		NAME_None,
+		true,
+		HAlign_Fill,
+		FNewMenuDelegate(),
+		Visibility
+	);
 	ToolbarBuilder.EndSection();
 
 #undef LOCTEXT_NAMESPACE
diff --git a/Engine/Plugins/Developer/RenderDocPlugin/Source/RenderDocPlugin/Private/SRenderDocPluginEditorExtension.h b/Engine/Plugins/Developer/RenderDocPlugin/Source/RenderDocPlugin/Private/SRenderDocPluginEditorExtension.h
--- a/Engine/Plugins/Developer/RenderDocPlugin/Source/RenderDocPlugin/Private/SRenderDocPluginEditorExtension.h
+++ b/Engine/Plugins/Developer/RenderDocPlugin/Source/RenderDocPlugin/Private/SRenderDocPluginEditorExtension.h
@@ -20,6 +20,7 @@ public:
 
 private:
   void Initialize(FRenderDocPluginModule* ThePlugin);
+  void ExtendToolbar();
   void AddToolbarExtension(FToolBarBuilder& ToolbarBuilder, FRenderDocPluginModule* ThePlugin);
 
   TSharedPtr<const FExtensionBase> ToolbarExtension;
diff --git a/Engine/Plugins/Developer/RenderDocPlugin/Source/RenderDocPlugin/RenderDocPlugin.Build.cs b/Engine/Plugins/Developer/RenderDocPlugin/Source/RenderDocPlugin/RenderDocPlugin.Build.cs
--- a/Engine/Plugins/Developer/RenderDocPlugin/Source/RenderDocPlugin/RenderDocPlugin.Build.cs
+++ b/Engine/Plugins/Developer/RenderDocPlugin/Source/RenderDocPlugin/RenderDocPlugin.Build.cs
@@ -34,6 +34,7 @@ namespace UnrealBuildTool.Rules
 					, "UnrealEd"
 					, "MainFrame"
 					, "GameProjectGeneration"
+					, "ToolMenus"
 				});
 			}
 
