RenderDocPlugin Help Crash Fix

#rnx

diff --git a/Engine/Plugins/Developer/RenderDocPlugin/Source/RenderDocPlugin/Private/SRenderDocPluginHelpWindow.cpp b/Engine/Plugins/Developer/RenderDocPlugin/Source/RenderDocPlugin/Private/SRenderDocPluginHelpWindow.cpp
--- a/Engine/Plugins/Developer/RenderDocPlugin/Source/RenderDocPlugin/Private/SRenderDocPluginHelpWindow.cpp
+++ b/Engine/Plugins/Developer/RenderDocPlugin/Source/RenderDocPlugin/Private/SRenderDocPluginHelpWindow.cpp
@@ -47,7 +47,7 @@ void SRenderDocPluginHelpWindow::Construct(const FArguments& InArgs)
 		.FocusWhenFirstShown(true)
 		.bDragAnywhere(false)
 		.ActivationPolicy(EWindowActivationPolicy::Always)
-		.ScreenPosition(FVector2D((float)(GEditor->GetActiveViewport()->GetSizeXY().X) / 2.0, (float)(GEditor->GetActiveViewport()->GetSizeXY().Y) / 2.0))
+		.ScreenPosition(FVector2D((float)(GSystemResolution.ResX) / 2.0, (float)(GSystemResolution.ResY) / 2.0))
 		[
 			SNew(SGridPanel)
 			.FillColumn(0, 0.2f)
