Fix issue where RenderDoc button didn't have a menu label.
- Allow the Campture Frame menu item to flow through the normal Command List system.
- Rename CaptureFrameCommand & Icon binding so that the command's icon can be found through the normal process.

#rnx

diff --git a/Engine/Plugins/Developer/RenderDocPlugin/Source/RenderDocPlugin/Private/RenderDocPluginCommands.cpp b/Engine/Plugins/Developer/RenderDocPlugin/Source/RenderDocPlugin/Private/RenderDocPluginCommands.cpp
--- a/Engine/Plugins/Developer/RenderDocPlugin/Source/RenderDocPlugin/Private/RenderDocPluginCommands.cpp
+++ b/Engine/Plugins/Developer/RenderDocPlugin/Source/RenderDocPlugin/Private/RenderDocPluginCommands.cpp
@@ -8,7 +8,7 @@
  
 void FRenderDocPluginCommands::RegisterCommands()
 {
-	UI_COMMAND(CaptureFrameCommand, "Capture Frame", "Captures the next frame and launches the renderdoc UI (Alt+F12)", EUserInterfaceActionType::Button, FInputChord(EKeys::F12, EModifierKey::Alt));
+	UI_COMMAND(CaptureFrame, "Capture Frame", "Captures the next frame and launches the renderdoc UI (Alt+F12)", EUserInterfaceActionType::Button, FInputChord(EKeys::F12, EModifierKey::Alt));
 }
 
 #undef LOCTEXT_NAMESPACE
diff --git a/Engine/Plugins/Developer/RenderDocPlugin/Source/RenderDocPlugin/Private/RenderDocPluginCommands.h b/Engine/Plugins/Developer/RenderDocPlugin/Source/RenderDocPlugin/Private/RenderDocPluginCommands.h
--- a/Engine/Plugins/Developer/RenderDocPlugin/Source/RenderDocPlugin/Private/RenderDocPluginCommands.h
+++ b/Engine/Plugins/Developer/RenderDocPlugin/Source/RenderDocPlugin/Private/RenderDocPluginCommands.h
@@ -17,7 +17,7 @@ public:
 
 	virtual void RegisterCommands() override;
 
-	TSharedPtr<class FUICommandInfo> CaptureFrameCommand;
+	TSharedPtr<class FUICommandInfo> CaptureFrame;
 };
 
 #endif
diff --git a/Engine/Plugins/Developer/RenderDocPlugin/Source/RenderDocPlugin/Private/RenderDocPluginStyle.cpp b/Engine/Plugins/Developer/RenderDocPlugin/Source/RenderDocPlugin/Private/RenderDocPluginStyle.cpp
--- a/Engine/Plugins/Developer/RenderDocPlugin/Source/RenderDocPlugin/Private/RenderDocPluginStyle.cpp
+++ b/Engine/Plugins/Developer/RenderDocPlugin/Source/RenderDocPlugin/Private/RenderDocPluginStyle.cpp
@@ -44,7 +44,7 @@ void FRenderDocPluginStyle::Initialize()
 	}
 
 	StyleSet->Set("RenderDocPlugin.Icon", new FSlateImageBrush(FRenderDocPluginStyle::InContent("Icon40", ".png"), FVector2D(40.0f, 40.0f)));
-	StyleSet->Set("RenderDocPlugin.CaptureFrameIcon", new FSlateImageBrush(FRenderDocPluginStyle::InContent("ViewportIcon16", ".png"), FVector2D(16.0f, 16.0f)));
+	StyleSet->Set("RenderDocPlugin.CaptureFrame", new FSlateImageBrush(FRenderDocPluginStyle::InContent("ViewportIcon16", ".png"), FVector2D(16.0f, 16.0f)));
 	
 	FSlateStyleRegistry::RegisterSlateStyle(*StyleSet.Get());
 };
diff --git a/Engine/Plugins/Developer/RenderDocPlugin/Source/RenderDocPlugin/Private/SRenderDocPluginEditorExtension.cpp b/Engine/Plugins/Developer/RenderDocPlugin/Source/RenderDocPlugin/Private/SRenderDocPluginEditorExtension.cpp
--- a/Engine/Plugins/Developer/RenderDocPlugin/Source/RenderDocPlugin/Private/SRenderDocPluginEditorExtension.cpp
+++ b/Engine/Plugins/Developer/RenderDocPlugin/Source/RenderDocPlugin/Private/SRenderDocPluginEditorExtension.cpp
@@ -37,7 +37,7 @@ public:
 	/** Widget constructor */
 	void Construct(const FArguments& Args)
 	{
-		FSlateIcon IconBrush = FSlateIcon(FRenderDocPluginStyle::Get()->GetStyleSetName(), "RenderDocPlugin.CaptureFrameIcon");
+		FSlateIcon IconBrush = FSlateIcon(FRenderDocPluginStyle::Get()->GetStyleSetName(), "RenderDocPlugin.CaptureFrame");
 		
 		ChildSlot
 		[
@@ -45,10 +45,10 @@ public:
 			.HAlign(HAlign_Center)
 			.VAlign(VAlign_Bottom)
 			.ContentPadding(FMargin(1.0f))
-			.ToolTipText(FRenderDocPluginCommands::Get().CaptureFrameCommand->GetDescription())
+			.ToolTipText(FRenderDocPluginCommands::Get().CaptureFrame->GetDescription())
 			.OnClicked_Lambda([this]() 
 				{ 
-					FPlayWorldCommands::GlobalPlayWorldActions->GetActionForCommand(FRenderDocPluginCommands::Get().CaptureFrameCommand)->Execute();
+					FPlayWorldCommands::GlobalPlayWorldActions->GetActionForCommand(FRenderDocPluginCommands::Get().CaptureFrame)->Execute();
 					return(FReply::Handled()); 
 				})
 			[
@@ -125,7 +125,7 @@ void FRenderDocPluginEditorExtension::Initialize(FRenderDocPluginModule* ThePlug
 		check(FPlayWorldCommands::GlobalPlayWorldActions.IsValid());
 
 		//Register the editor hotkeys
-		FPlayWorldCommands::GlobalPlayWorldActions->MapAction(FRenderDocPluginCommands::Get().CaptureFrameCommand,
+		FPlayWorldCommands::GlobalPlayWorldActions->MapAction(FRenderDocPluginCommands::Get().CaptureFrame,
 			FExecuteAction::CreateLambda([]()
 		{
 			FRenderDocPluginModule& PluginModule = FModuleManager::GetModuleChecked<FRenderDocPluginModule>("RenderDocPlugin");
@@ -148,17 +148,8 @@ void FRenderDocPluginEditorExtension::ExtendToolbar()
 	UToolMenu* Menu = UToolMenus::Get()->ExtendMenu("LevelEditor.ViewportToolbar");
 	
 	FToolMenuSection& RightSection = Menu->FindOrAddSection("Right");
-	FToolMenuEntry& Entry = RightSection.AddMenuEntry(
-		"RenderDocButton",
-		FText::GetEmpty(),
-		FRenderDocPluginCommands::Get().CaptureFrameCommand->GetDescription(),
-		FSlateIcon(FRenderDocPluginStyle::Get()->GetStyleSetName(), "RenderDocPlugin.CaptureFrameIcon"),
-		FUIAction(FExecuteAction::CreateLambda([]
-		{
-			FPlayWorldCommands::GlobalPlayWorldActions->GetActionForCommand(FRenderDocPluginCommands::Get().CaptureFrameCommand)->Execute();
-		}))
-	);
-	
+	FToolMenuEntry& Entry = RightSection.AddMenuEntry(FRenderDocPluginCommands::Get().CaptureFrame);
+	Entry.ToolBarData.LabelOverride = FText::GetEmpty();
 	Entry.InsertPosition.Position = EToolMenuInsertType::First;
 }
 
