Fixed programmatic RenderDoc captures randomly crashing

diff --git a/Engine/Plugins/Developer/RenderDocPlugin/Source/RenderDocPlugin/Private/RenderDocPluginModule.cpp b/Engine/Plugins/Developer/RenderDocPlugin/Source/RenderDocPlugin/Private/RenderDocPluginModule.cpp
--- a/Engine/Plugins/Developer/RenderDocPlugin/Source/RenderDocPlugin/Private/RenderDocPluginModule.cpp
+++ b/Engine/Plugins/Developer/RenderDocPlugin/Source/RenderDocPlugin/Private/RenderDocPluginModule.cpp
@@ -22,6 +22,7 @@
 #if WITH_EDITOR
 #include "Editor.h"
 #include "Editor/EditorEngine.h"
+#include "Interfaces/IMainFrameModule.h"
 #include "ISettingsModule.h"
 #include "ISettingsSection.h"
 #include "LevelEditor.h"
@@ -326,15 +327,31 @@ void FRenderDocPluginModule::StartupModule()
 }
 
 void FRenderDocPluginModule::OnPostEngineInit()
-	{
+{
 #if WITH_EDITOR
 	if (FSlateApplication::IsInitialized() && !IsRunningCommandlet())
 	{
 		EditorExtension = MakeShared<FRenderDocPluginEditorExtension>(this);
+
+		// Cannot register in StartupModule because the RDoc plugin is initialized earlier than the main frame's
+		if (IMainFrameModule* MainFrameModule = FModuleManager::GetModulePtr<IMainFrameModule>("MainFrame"))
+		{
+			MainFrameModule->OnMainFrameCreationFinished().AddRaw(this, &FRenderDocPluginModule::OnMainFrameCreationFinished);
+		}
 	}
 #endif // WITH_EDITOR
 }
 
+#if WITH_EDITOR
+void FRenderDocPluginModule::OnMainFrameCreationFinished(TSharedPtr<SWindow> ParentWindow, bool bIsRunningStartupDialog)
+{
+	if (ParentWindow.IsValid())
+	{
+		MainWindowHandle = ParentWindow->GetNativeWindow()->GetOSWindowHandle();
+	}
+}
+#endif // WITH_EDITOR
+
 void FRenderDocPluginModule::BeginFrameCapture()
 {
 	UE_LOG(RenderDocPlugin, Log, TEXT("Capture frame and launch renderdoc!"));
@@ -647,6 +664,11 @@ void FRenderDocPluginModule::ShutdownModule()
 	IModularFeatures::Get().UnregisterModularFeature(IRenderCaptureProvider::GetModularFeatureName(), (IRenderCaptureProvider*)this);
 
 #if WITH_EDITOR
+	if (IMainFrameModule* MainFrameModule = FModuleManager::GetModulePtr<IMainFrameModule>("MainFrame"))
+	{
+		MainFrameModule->OnMainFrameCreationFinished().RemoveAll(this);
+	}
+
 	EditorExtension.Reset();
 #endif // WITH_EDITOR
 
@@ -675,7 +697,7 @@ void FRenderDocPluginModule::BeginCapture_RenderThread(FRHICommandListImmediate*
 	InRHICommandList->SubmitCommandsAndFlushGPU();
 	InRHICommandList->EnqueueLambda([this, Device](FRHICommandListImmediate& RHICommandList)
 	{
-		RenderDocAPI->StartFrameCapture(Device, NULL);
+		RenderDocAPI->StartFrameCapture(Device, MainWindowHandle);
 	});
 }
 
@@ -685,7 +707,7 @@ void FRenderDocPluginModule::EndCapture_RenderThread(FRHICommandListImmediate* I
 	InRHICommandList->SubmitCommandsAndFlushGPU();
 	InRHICommandList->EnqueueLambda([this, Device, InFlags, DestFileName = InDestFileName](FRHICommandListImmediate& RHICommandList)
 	{
-		uint32 Result = RenderDocAPI->EndFrameCapture(Device, NULL);
+		uint32 Result = RenderDocAPI->EndFrameCapture(Device, MainWindowHandle);
 		if (Result == 1)
 		{
 			FRenderDocFrameCapturer::SaveAndLaunch(this, InFlags, DestFileName);
diff --git a/Engine/Plugins/Developer/RenderDocPlugin/Source/RenderDocPlugin/Private/RenderDocPluginModule.h b/Engine/Plugins/Developer/RenderDocPlugin/Source/RenderDocPlugin/Private/RenderDocPluginModule.h
--- a/Engine/Plugins/Developer/RenderDocPlugin/Source/RenderDocPlugin/Private/RenderDocPluginModule.h
+++ b/Engine/Plugins/Developer/RenderDocPlugin/Source/RenderDocPlugin/Private/RenderDocPluginModule.h
@@ -8,6 +8,7 @@
 DECLARE_LOG_CATEGORY_EXTERN(RenderDocPlugin, Log, All);
 
 class FRenderDocPluginEditorExtension;
+class SWindow;
 
 class FRenderDocPluginModule : public IRenderDocPlugin
 {
@@ -42,6 +43,7 @@ private:
 	void CaptureFrame();
 
 #if WITH_EDITOR
+	void OnMainFrameCreationFinished(TSharedPtr<SWindow> ParentWindow, bool bIsRunningStartupDialog);
 	void CapturePIE(const TArray<FString>& Args);
 #endif // WITH_EDITOR
 
@@ -63,7 +65,8 @@ private:
 private:
 	FRenderDocPluginLoader Loader;
 	FRenderDocPluginLoader::RENDERDOC_API_CONTEXT* RenderDocAPI;
-	
+
+	void* MainWindowHandle = nullptr; // Keep a handle of the main window so that programmatic captures (i.e. BeginCapture_RenderThread / EndCapture_RenderThread) can still pass a valid window handle to RenderDoc
 	uint64 DelayedCaptureTick; // Tracks on which frame a delayed capture should trigger, if any (when bCaptureDelayInSeconds == false)
 	double DelayedCaptureSeconds; // Tracks at which time a delayed capture should trigger, if any (when bCaptureDelayInSeconds == true)
 	uint64 CaptureFrameCount; // Tracks how many frames should be captured
